Added new experimental plugin AVCodecs which handles software/hardware encoding and decoding of audio and video

diff --git a/Engine/Plugins/Experimental/AVCodecs/AVCodecsCore/Source/AVCodecsCore/Private/Video/CodecUtils/CodecUtilsH264.cpp b/Engine/Plugins/Experimental/AVCodecs/AVCodecsCore/Source/AVCodecsCore/Private/Video/CodecUtils/CodecUtilsH264.cpp
new file mode 100644
--- /dev/null
+++ b/Engine/Plugins/Experimental/AVCodecs/AVCodecsCore/Source/AVCodecsCore/Private/Video/CodecUtils/CodecUtilsH264.cpp
@@ -0,0 +1,804 @@
+// Copyright Epic Games, Inc. All Rights Reserved.
+
+#include "Video/CodecUtils/CodecUtilsH264.h"
+
+#include "Containers/Array.h"
+#include "AVResult.h"
+#include "Video/CodecUtils/CodecUtilsH265.h"
+
+FH264ProfileDefinition GH264ProfileDefinitions[static_cast<uint8>(EH264Profile::MAX)] = {
+	{ EH264Profile::Auto, UE::AVCodecCore::H264::EH264ProfileIDC::Auto, UE::AVCodecCore::H264::EH264ConstraintFlag::None, TEXT("Auto") },
+	{ EH264Profile::CALVLC444Intra, UE::AVCodecCore::H264::EH264ProfileIDC::CALVLC444Intra, UE::AVCodecCore::H264::EH264ConstraintFlag::None, TEXT("CALVLC444Intra") },
+	{ EH264Profile::Baseline, UE::AVCodecCore::H264::EH264ProfileIDC::Baseline, UE::AVCodecCore::H264::EH264ConstraintFlag::None, TEXT("Baseline") },
+	{ EH264Profile::ConstrainedBaseline, UE::AVCodecCore::H264::EH264ProfileIDC::Baseline, UE::AVCodecCore::H264::EH264ConstraintFlag::Set1, TEXT("ConstrainedBaseline") },
+	{ EH264Profile::Main, UE::AVCodecCore::H264::EH264ProfileIDC::Main, UE::AVCodecCore::H264::EH264ConstraintFlag::None, TEXT("Main") },
+	{ EH264Profile::ScalableBaseline, UE::AVCodecCore::H264::EH264ProfileIDC::ScalableBaseline, UE::AVCodecCore::H264::EH264ConstraintFlag::None, TEXT("ScalableBaseline") },
+	{ EH264Profile::ScalableConstrainedBaseline, UE::AVCodecCore::H264::EH264ProfileIDC::ScalableBaseline, UE::AVCodecCore::H264::EH264ConstraintFlag::Set5, TEXT("ScalableConstrainedBaseline") },
+	{ EH264Profile::ScalableHigh, UE::AVCodecCore::H264::EH264ProfileIDC::ScalableHigh, UE::AVCodecCore::H264::EH264ConstraintFlag::None, TEXT("ScalableHigh") },
+	{ EH264Profile::ScalableConstrainedHigh, UE::AVCodecCore::H264::EH264ProfileIDC::ScalableHigh, UE::AVCodecCore::H264::EH264ConstraintFlag::Set5, TEXT("ScalableConstrainedHigh") },
+	{ EH264Profile::ScalableHighIntra, UE::AVCodecCore::H264::EH264ProfileIDC::ScalableHigh, UE::AVCodecCore::H264::EH264ConstraintFlag::Set3, TEXT("ScalableHighIntra") },
+	{ EH264Profile::Extended, UE::AVCodecCore::H264::EH264ProfileIDC::Extended, UE::AVCodecCore::H264::EH264ConstraintFlag::None, TEXT("Extended") },
+	{ EH264Profile::High, UE::AVCodecCore::H264::EH264ProfileIDC::High, UE::AVCodecCore::H264::EH264ConstraintFlag::None, TEXT("High") },
+	{ EH264Profile::ProgressiveHigh, UE::AVCodecCore::H264::EH264ProfileIDC::High, UE::AVCodecCore::H264::EH264ConstraintFlag::Set4, TEXT("ProgressiveHigh") },
+	{ EH264Profile::ConstrainedHigh, UE::AVCodecCore::H264::EH264ProfileIDC::High, UE::AVCodecCore::H264::EH264ConstraintFlag::Set4& UE::AVCodecCore::H264::EH264ConstraintFlag::Set5, TEXT("ConstrainedHigh") },
+	{ EH264Profile::High10, UE::AVCodecCore::H264::EH264ProfileIDC::High10, UE::AVCodecCore::H264::EH264ConstraintFlag::None, TEXT("High10") },
+	{ EH264Profile::High10Intra, UE::AVCodecCore::H264::EH264ProfileIDC::High10, UE::AVCodecCore::H264::EH264ConstraintFlag::Set3, TEXT("High10Intra") },
+	{ EH264Profile::MultiviewHigh, UE::AVCodecCore::H264::EH264ProfileIDC::MultiviewHigh, UE::AVCodecCore::H264::EH264ConstraintFlag::None, TEXT("MultiviewHigh") },
+	{ EH264Profile::High422, UE::AVCodecCore::H264::EH264ProfileIDC::High422, UE::AVCodecCore::H264::EH264ConstraintFlag::None, TEXT("High422") },
+	{ EH264Profile::High422Intra, UE::AVCodecCore::H264::EH264ProfileIDC::High422, UE::AVCodecCore::H264::EH264ConstraintFlag::Set3, TEXT("High422Intra") },
+	{ EH264Profile::StereoHigh, UE::AVCodecCore::H264::EH264ProfileIDC::StereoHigh, UE::AVCodecCore::H264::EH264ConstraintFlag::None, TEXT("StereoHigh") },
+	{ EH264Profile::MultiviewDepthHigh, UE::AVCodecCore::H264::EH264ProfileIDC::MultiviewDepthHigh, UE::AVCodecCore::H264::EH264ConstraintFlag::None, TEXT("MultiviewDepthHigh") },
+	{ EH264Profile::High444, UE::AVCodecCore::H264::EH264ProfileIDC::High444, UE::AVCodecCore::H264::EH264ConstraintFlag::None, TEXT("High444") },
+	{ EH264Profile::High444Intra, UE::AVCodecCore::H264::EH264ProfileIDC::High444, UE::AVCodecCore::H264::EH264ConstraintFlag::Set3, TEXT("High444Intra") }
+};
+
+namespace UE::AVCodecCore::H264
+{
+	FAVResult FindNALUs(FVideoPacket const& InPacket, TArray<FNaluInfo>& FoundNalus)
+	{
+		FBitstreamReader Reader(InPacket.DataPtr.Get(), InPacket.DataSize);
+		if (Reader.NumBytesRemaining() < 3)
+		{
+			return FAVResult(EAVResult::Warning, TEXT("Bitstream not long enough to hold a NALU"), TEXT("H264"));
+		}
+
+		TArrayView64<uint8> const Data = InPacket.GetData();
+		
+		// Skip over stream in intervals of 3 until Data[i + 2] is either 0 or 1
+		for (int64 i = 0; i < Data.Num() - 2;)
+		{
+			if (Data[i + 2] > 1)
+			{
+				i += 3;
+			}
+			else if (Data[i + 2] == 1)
+			{
+				if (Data[i + 1] == 0 && Data[i] == 0)
+				{
+					// Found start sequence of NalU but we don't know if it has a 3 or 4 byte start code so we check.
+					FNaluInfo NalInfo = { (uint64)i, 0, 3, 0, ENaluType::Unspecified, nullptr };
+					if (NalInfo.Start > 0 && Data[NalInfo.Start - 1] == 0)
+					{
+						++NalInfo.StartCodeSize;
+						--NalInfo.Start;
+					}
+
+					FBitstreamReader Bitstream(&Data[NalInfo.Start + NalInfo.StartCodeSize], 1);
+					verifyf(Bitstream.ReadBits(1) == 0, TEXT("Forbidden Zero bit not Zero in NAL Header"));
+
+					NalInfo.RefIdc = Bitstream.ReadBits(2);
+					NalInfo.Type = (ENaluType)Bitstream.ReadBits(5);
+
+					NalInfo.Data = &Data[NalInfo.Start + NalInfo.StartCodeSize + 1];
+
+					// Update length of previous entry.
+					if (FoundNalus.Num() > 0)
+					{
+						FoundNalus.Last().Size = NalInfo.Start - (FoundNalus.Last().Start + FoundNalus.Last().StartCodeSize);
+					}
+
+					FoundNalus.Add(NalInfo);
+				}
+
+				i += 3;
+			}
+			else
+			{
+				++i;
+			}
+		}
+
+		if (FoundNalus.Num() == 0)
+		{
+			return FAVResult(EAVResult::PendingInput, TEXT("no NALUs found in BitDataStream"), TEXT("BitstreamParserH264"));
+		}
+
+		// Last Nal size is the remaining size of the bitstream minus a trailing zero byte
+		FoundNalus.Last().Size = Data.Num() - (FoundNalus.Last().Start + FoundNalus.Last().StartCodeSize + 1);
+
+	#if !IS_PROGRAM
+
+		FAVResult::Log(EAVResult::Success, FString::Printf(TEXT("FAVExtension::TransformConfig found %d NALUs in bitdatastream"), FoundNalus.Num()));
+
+		for (const FNaluInfo& NalUInfo : FoundNalus)
+		{
+			FAVResult::Log(EAVResult::Success, FString::Printf(TEXT("Found NALU at %llu size %llu with type %u"), NalUInfo.Start, NalUInfo.Size, (uint8)NalUInfo.Type));
+		}
+
+	#endif //!IS_PROGRAM
+
+		return EAVResult::Success;
+	}
+
+	FAVResult ParseSEI(FBitstreamReader& Bitstream, FNaluInfo const& InNaluInfo, SEI_t& OutSEI)
+	{
+		//unimplemented();
+		return FAVResult(EAVResult::Success, TEXT("SEI Unimplemented"));
+	}
+
+	void ParseScalingList(FBitstreamReader& Bitstream, const uint8& chroma_format_idc,  FNalu::U<1> scaling_list_present_flag[12], uint8 ScalingList4x4[6][16], uint8 ScalingList8x8[6][64])
+	{
+		auto scaling_list = [&Bitstream](uint8* scalingList, int32 sizeOfScalingList, bool& useDefaultScalingMatrixFlag) -> void
+		{
+			int32 lastScale = 8;
+			int32 nextScale = 8;
+			for (int32 j = 0; j < sizeOfScalingList; ++j)
+			{
+				if (nextScale)
+				{
+					FNalu::SE delta_scale;
+					Bitstream.Read(delta_scale);
+					
+					nextScale = (lastScale + delta_scale + 256) % 256;
+					useDefaultScalingMatrixFlag = (j == 0 && nextScale == 0);
+				}
+				scalingList[j] = (nextScale == 0) ? lastScale : nextScale;
+				lastScale = scalingList[j];
+			}
+		};
+
+		bool useDefaultScalingMatrixFlag = false;
+		for (int32 i = 0, iMax = (chroma_format_idc != 3) ? 8 : 12; i < iMax; ++i)
+		{
+			Bitstream.Read(scaling_list_present_flag[i]);
+			if (scaling_list_present_flag[i])
+			{
+				if (i < 6)
+				{
+					scaling_list(ScalingList4x4[i], 16, useDefaultScalingMatrixFlag);
+					if (useDefaultScalingMatrixFlag)
+					{
+						if (i < 3)
+						{
+							FMemory::Memcpy(ScalingList4x4[i], Default_4x4_Intra, 16);
+						}
+						else
+						{
+							FMemory::Memcpy(ScalingList4x4[i], Default_4x4_Inter, 16);									
+						}
+					}
+				}
+				else
+				{
+					scaling_list(ScalingList8x8[i - 6], 64, useDefaultScalingMatrixFlag);
+					if (useDefaultScalingMatrixFlag)
+					{
+						if (i % 2)
+						{
+							FMemory::Memcpy(ScalingList8x8[i - 6], Default_8x8_Intra, 64);
+						}
+						else
+						{
+							FMemory::Memcpy(ScalingList8x8[i - 6], Default_8x8_Inter, 64);									
+						}
+					}
+				}
+			}
+		}
+	}
+
+	FAVResult ParseSPS(FBitstreamReader& Bitstream, FNaluInfo const& InNaluInfo, TMap<uint32, SPS_t>& OutMapSPS)
+	{
+		// Be sure to maintain parity with SPS struct values, or we may miss bits
+		FNalu::U<8, EH264ProfileIDC> temp_profile_idc;
+		FNalu::U<8, EH264ConstraintFlag> temp_constraint_flags;
+
+		Bitstream.Read(
+			temp_profile_idc,
+			temp_constraint_flags);
+		
+		FNalu::SE seq_parameter_set_id;
+		Bitstream.Read(seq_parameter_set_id);
+
+		// Find SPS at current ID or set it to the map
+		SPS_t& OutSPS = OutMapSPS.FindOrAdd(seq_parameter_set_id);
+
+		OutSPS.profile_idc = temp_profile_idc;
+		OutSPS.constraint_flags = temp_constraint_flags;
+		
+		Bitstream.Read(
+			OutSPS.level_idc,
+			OutSPS.seq_parameter_set_id);
+
+		if (OutSPS.profile_idc == EH264ProfileIDC::High ||
+			OutSPS.profile_idc == EH264ProfileIDC::High10 ||
+			OutSPS.profile_idc == EH264ProfileIDC::High422 ||
+			OutSPS.profile_idc == EH264ProfileIDC::High444 ||
+			OutSPS.profile_idc == EH264ProfileIDC::CALVLC444Intra ||
+			OutSPS.profile_idc == EH264ProfileIDC::ScalableBaseline ||
+			OutSPS.profile_idc == EH264ProfileIDC::ScalableHigh ||
+			OutSPS.profile_idc == EH264ProfileIDC::MultiviewHigh ||
+			OutSPS.profile_idc == EH264ProfileIDC::MultiviewDepthHigh)
+		{
+			Bitstream.Read(OutSPS.chroma_format_idc);
+
+			if (OutSPS.chroma_format_idc == 3)
+			{
+				Bitstream.Read(OutSPS.separate_colour_plane_flag);
+			}
+
+			Bitstream.Read(
+				OutSPS.bit_depth_luma_minus8,
+				OutSPS.bit_depth_chroma_minus8,
+				OutSPS.qpprime_y_zero_transform_bypass_flag,
+				OutSPS.seq_scaling_matrix_present_flag);
+
+			if (OutSPS.seq_scaling_matrix_present_flag)
+			{
+				ParseScalingList(Bitstream, OutSPS.chroma_format_idc, OutSPS.seq_scaling_list_present_flag, OutSPS.ScalingList4x4, OutSPS.ScalingList8x8);
+			}
+		}
+
+		Bitstream.Read(OutSPS.log2_max_frame_num_minus4);
+
+		Bitstream.Read(OutSPS.pic_order_cnt_type);
+		if (OutSPS.pic_order_cnt_type == 0)
+		{
+			Bitstream.Read(OutSPS.log2_max_pic_order_cnt_lsb_minus4);
+		}
+		else if (OutSPS.pic_order_cnt_type == 1)
+		{
+			Bitstream.Read(
+				OutSPS.delta_pic_order_always_zero_flag,
+				OutSPS.offset_for_non_ref_pic,
+				OutSPS.offset_for_top_to_bottom_field,
+				OutSPS.num_ref_frames_in_pic_order_cnt_cycle);
+			
+			for (uint32 i = 0; i < OutSPS.num_ref_frames_in_pic_order_cnt_cycle; ++i)
+			{
+				Bitstream.Read(OutSPS.offset_for_ref_frame[i]);
+			}
+		}
+
+		Bitstream.Read(
+			OutSPS.max_num_ref_frames,
+			OutSPS.gaps_in_frame_num_value_allowed_flag,
+			OutSPS.pic_width_in_mbs_minus1,
+			OutSPS.pic_height_in_map_units_minus1);
+
+		Bitstream.Read(OutSPS.frame_mbs_only_flag);
+		if (!OutSPS.frame_mbs_only_flag)
+		{
+			Bitstream.Read(OutSPS.mb_adaptive_frame_field_flag);
+		}
+
+		Bitstream.Read(OutSPS.direct_8x8_inference_flag);
+
+		Bitstream.Read(OutSPS.frame_cropping_flag);
+		if (OutSPS.frame_cropping_flag)
+		{
+			Bitstream.Read(
+				OutSPS.frame_crop_left_offset,
+				OutSPS.frame_crop_right_offset,
+				OutSPS.frame_crop_top_offset,
+				OutSPS.frame_crop_bottom_offset);
+		}
+
+		Bitstream.Read(OutSPS.vui_parameters_present_flag);
+		if (OutSPS.vui_parameters_present_flag)
+		{
+			Bitstream.Read(OutSPS.aspect_ratio_info_present_flag);
+			if (OutSPS.aspect_ratio_info_present_flag)
+			{
+				OutSPS.aspect_ratio_idc = (EH264AspectRatioIDC)Bitstream.ReadBits(8);
+				if (OutSPS.aspect_ratio_idc == EH264AspectRatioIDC::Extended_SAR)
+				{
+					Bitstream.Read(OutSPS.sar_width);
+					Bitstream.Read(OutSPS.sar_height);
+				}
+			}
+
+			Bitstream.Read(OutSPS.overscan_info_present_flag);
+			if (OutSPS.overscan_info_present_flag)
+			{
+				Bitstream.Read(OutSPS.overscan_appropriate_flag);
+			}
+
+			Bitstream.Read(OutSPS.video_signal_type_present_flag);
+			if (OutSPS.video_signal_type_present_flag)
+			{
+				Bitstream.Read(
+					OutSPS.video_format,
+					OutSPS.video_full_range_flag);
+
+				Bitstream.Read(OutSPS.colour_description_present_flag);
+				if (OutSPS.colour_description_present_flag)
+				{
+					Bitstream.Read(
+						OutSPS.colour_primaries,
+						OutSPS.transfer_characteristics,
+						OutSPS.matrix_coefficients);
+				}
+			}
+
+			Bitstream.Read(OutSPS.chroma_loc_info_present_flag);
+			if (OutSPS.chroma_loc_info_present_flag)
+			{
+				Bitstream.Read(
+					OutSPS.chroma_sample_loc_type_top_field,
+					OutSPS.chroma_sample_loc_type_bottom_field);
+			}
+			
+			Bitstream.Read(OutSPS.timing_info_present_flag);
+			if (OutSPS.timing_info_present_flag)
+			{
+				Bitstream.Read(
+					OutSPS.num_units_in_tick,
+					OutSPS.time_scale,
+					OutSPS.fixed_frame_rate_flag);
+			}
+
+			auto hrd_parameters = [&Bitstream, &OutSPS]() -> void
+			{
+				Bitstream.Read(
+					OutSPS.cpb_cnt_minus1,
+					OutSPS.bit_rate_scale,
+					OutSPS.cpb_size_scale);
+				
+				for (uint32 SchedSelIdx = 0; SchedSelIdx <= OutSPS.cpb_cnt_minus1; SchedSelIdx++ )
+				{
+					Bitstream.Read(
+						OutSPS.bit_rate_value_minus1[ SchedSelIdx ],
+						OutSPS.cpb_size_value_minus1[ SchedSelIdx ],
+						OutSPS.cbr_flag[ SchedSelIdx ]);
+				}
+
+				Bitstream.Read(
+					OutSPS.initial_cpb_removal_delay_length_minus1,
+					OutSPS.cpb_removal_delay_length_minus1,
+					OutSPS.dpb_output_delay_length_minus1,
+					OutSPS.time_offset_length);
+			};
+
+			Bitstream.Read(OutSPS.nal_hrd_parameters_present_flag);
+			if (OutSPS.nal_hrd_parameters_present_flag)
+			{
+				hrd_parameters();
+			}
+			
+			Bitstream.Read(OutSPS.vcl_hrd_parameters_present_flag);
+			if (OutSPS.vcl_hrd_parameters_present_flag)
+			{
+				hrd_parameters();
+			}
+
+			if (OutSPS.nal_hrd_parameters_present_flag || OutSPS.vcl_hrd_parameters_present_flag)
+			{
+				Bitstream.Read(OutSPS.low_delay_hrd_flag);
+			}
+
+			Bitstream.Read(OutSPS.pic_struct_present_flag);
+			Bitstream.Read(OutSPS.bitstream_restriction_flag);
+			if( OutSPS.bitstream_restriction_flag )
+			{
+				Bitstream.Read(
+					OutSPS.motion_vectors_over_pic_boundaries_flag,
+					OutSPS.max_bytes_per_pic_denom,
+					OutSPS.max_bits_per_mb_denom,
+					OutSPS.log2_max_mv_length_horizontal,
+					OutSPS.log2_max_mv_length_vertical,
+					OutSPS.max_num_reorder_frames,
+					OutSPS.max_dec_frame_buffering);
+			}
+		}
+
+		return EAVResult::Success;
+	}
+
+	FAVResult ParsePPS(FBitstreamReader& Bitstream, FNaluInfo const& InNaluInfo, TMap<uint32, SPS_t> const& InMapSPS, TMap<uint32, PPS_t>& OutMapPPS)
+	{	
+		FNalu::UE pic_parameter_set_id;
+		Bitstream.Read(pic_parameter_set_id);
+
+		PPS_t& OutPPS = OutMapPPS.FindOrAdd(pic_parameter_set_id);
+		OutPPS.pic_parameter_set_id = pic_parameter_set_id;
+		
+		Bitstream.Read(
+			OutPPS.seq_parameter_set_id,
+			OutPPS.entropy_coding_mode_flag,
+			OutPPS.bottom_field_pic_order_in_frame_present_flag,
+			OutPPS.num_slice_groups_minus1);
+		
+		if( OutPPS.num_slice_groups_minus1 > 0 )
+		{
+			Bitstream.Read(OutPPS.slice_group_map_type);
+			if( OutPPS.slice_group_map_type == 0 )
+			{
+				OutPPS.run_length_minus1.AddUninitialized(OutPPS.num_slice_groups_minus1 + 1);
+				for (uint32 iGroup = 0; iGroup <= OutPPS.num_slice_groups_minus1; iGroup++ )
+				{
+					Bitstream.Read(OutPPS.run_length_minus1[ iGroup ]);
+				}
+			}
+			else if( OutPPS.slice_group_map_type == 2 )
+			{
+				OutPPS.top_left.AddUninitialized(OutPPS.num_slice_groups_minus1 + 1);
+				OutPPS.bottom_right.AddUninitialized(OutPPS.num_slice_groups_minus1 + 1);
+				for (uint32 iGroup = 0; iGroup < OutPPS.num_slice_groups_minus1; iGroup++ )
+					{
+						Bitstream.Read(
+							OutPPS.top_left[ iGroup ],
+							OutPPS.bottom_right[ iGroup ]);
+					}
+			}
+			else if( OutPPS.slice_group_map_type == 3 ||
+					 OutPPS.slice_group_map_type == 4 ||
+					 OutPPS.slice_group_map_type == 5 )
+			{
+				Bitstream.Read(
+					OutPPS.slice_group_change_direction_flag,
+					OutPPS.slice_group_change_rate_minus1);
+			}
+			else if( OutPPS.slice_group_map_type == 6 )
+			{
+				Bitstream.Read(OutPPS.pic_size_in_map_units_minus1);
+				unsigned long numBits = 0;
+			#if defined(_MSC_VER)
+					_BitScanForward(&numBits, OutPPS.pic_size_in_map_units_minus1 + 1);
+			#else
+					numBits = __builtin_ctz(OutPPS.pic_size_in_map_units_minus1 + 1);
+			#endif
+				for( uint32 i = 0; i <= OutPPS.pic_size_in_map_units_minus1; i++ )
+				{
+					Bitstream.ReadBits(OutPPS.slice_group_id[i], numBits);
+				}
+			}
+		}
+		
+		Bitstream.Read(OutPPS.num_ref_idx_l0_default_active_minus1,
+			OutPPS.num_ref_idx_l1_default_active_minus1,
+			OutPPS.weighted_pred_flag,
+			OutPPS.weighted_bipred_idc,
+			OutPPS.pic_init_qp_minus26,
+			OutPPS.pic_init_qs_minus26,
+			OutPPS.chroma_qp_index_offset,
+			OutPPS.deblocking_filter_control_present_flag,
+			OutPPS.constrained_intra_pred_flag,
+			OutPPS.redundant_pic_cnt_present_flag);
+
+		auto more_rbsp_data = [&Bitstream]() -> bool
+		{
+			if(Bitstream.NumBitsRemaining() > 0)
+			{
+				return true;
+			}
+
+			return false;
+		};
+
+		if( more_rbsp_data() )
+		{
+			Bitstream.Read(OutPPS.transform_8x8_mode_flag);
+			
+			Bitstream.Read(OutPPS.pic_scaling_matrix_present_flag);
+			if (OutPPS.pic_scaling_matrix_present_flag)
+			{
+				ParseScalingList(Bitstream, InMapSPS[OutPPS.seq_parameter_set_id].chroma_format_idc, OutPPS.pic_scaling_list_present_flag, OutPPS.ScalingList4x4, OutPPS.ScalingList8x8);
+			}
+
+			Bitstream.Read(OutPPS.second_chroma_qp_index_offset);
+		}
+		
+		return EAVResult::Success;
+	}
+
+	FAVResult ParseSliceHeader(FBitstreamReader& Bitstream, FNaluInfo const& InNaluInfo, TMap<uint32, SPS_t> InMapSPS, TMap<uint32, PPS_t> const& InMapPPS, Slice_t& OutSlice)
+	{
+		Bitstream.Read(
+			OutSlice.first_mb_in_slice,
+			OutSlice.slice_type,
+			OutSlice.pic_parameter_set_id);
+
+		const PPS_t& CurrentPPS = InMapPPS[OutSlice.pic_parameter_set_id];
+		const SPS_t& CurrentSPS = InMapSPS[CurrentPPS.seq_parameter_set_id];
+
+		if (CurrentSPS.separate_colour_plane_flag)
+		{
+			Bitstream.Read(OutSlice.colour_plane_id);
+		}
+
+		Bitstream.ReadBits(OutSlice.frame_num, CurrentSPS.log2_max_frame_num_minus4 + 4);
+
+		if (!CurrentSPS.frame_mbs_only_flag)
+		{
+			Bitstream.Read(OutSlice.field_pic_flag);
+			if (OutSlice.field_pic_flag)
+			{
+				Bitstream.Read(OutSlice.bottom_field_flag);
+			}
+		}
+
+		if (InNaluInfo.Type == ENaluType::SliceIdrPicture)
+		{
+			Bitstream.Read(OutSlice.idr_pic_id);
+		}
+
+		if (CurrentSPS.pic_order_cnt_type == 0)
+		{
+			Bitstream.ReadBits(OutSlice.pic_order_cnt_lsb, CurrentSPS.log2_max_pic_order_cnt_lsb_minus4 + 4);
+
+			if (CurrentPPS.bottom_field_pic_order_in_frame_present_flag && !OutSlice.field_pic_flag)
+			{
+				Bitstream.Read(OutSlice.delta_pic_order_cnt_bottom);
+			}
+		}
+
+		if (CurrentSPS.pic_order_cnt_type == 1 && !CurrentSPS.delta_pic_order_always_zero_flag)
+		{
+			Bitstream.Read(OutSlice.delta_pic_order_cnt[0]);
+
+			if (CurrentPPS.bottom_field_pic_order_in_frame_present_flag && !OutSlice.field_pic_flag)
+			{
+				Bitstream.Read(OutSlice.delta_pic_order_cnt[1]);
+			}
+		}
+
+		if (CurrentPPS.redundant_pic_cnt_present_flag)
+		{
+			Bitstream.Read(OutSlice.redundant_pic_cnt);
+		}
+
+		const bool IsP = OutSlice.slice_type == 0 || OutSlice.slice_type == 5;
+		const bool IsB = OutSlice.slice_type == 1 || OutSlice.slice_type == 6;
+		const bool IsI = OutSlice.slice_type == 2 || OutSlice.slice_type == 7;
+		const bool IsSP = OutSlice.slice_type == 3 || OutSlice.slice_type == 8;
+		const bool IsSI = OutSlice.slice_type == 4 || OutSlice.slice_type == 9;
+		
+		if (IsB)
+		{
+			Bitstream.Read(OutSlice.direct_spatial_mv_pred_flag);
+		}
+
+		if ( IsP || IsSP || IsB )
+		{
+			Bitstream.Read(OutSlice.num_ref_idx_active_override_flag);
+
+			if (OutSlice.num_ref_idx_active_override_flag)
+			{
+				Bitstream.Read(OutSlice.num_ref_idx_l0_active_minus1);
+
+				if ( IsB )
+				{
+					Bitstream.Read(OutSlice.num_ref_idx_l1_active_minus1);
+				}
+			}
+		}
+
+		// Slice extensions
+		if (InNaluInfo.Type == ENaluType::SliceExtension || InNaluInfo.Type == ENaluType::SliceExtensionForDepthView)
+		{
+			// TODO ref_pic_list_mvc_modification()
+		}
+		else
+		{
+			if ( OutSlice.slice_type % 5 != 2 &&  OutSlice.slice_type % 5 != 4 )
+			{
+				Bitstream.Read(OutSlice.ref_pic_list_modification_flag_l0);
+				if (OutSlice.ref_pic_list_modification_flag_l0)
+				{
+					uint32 modification_of_pic_nums_idc;
+					do
+					{
+						Bitstream.Read(modification_of_pic_nums_idc);
+						if (modification_of_pic_nums_idc == 0 || modification_of_pic_nums_idc == 1)
+						{
+							FNalu::UE pic_num;
+							Bitstream.Read(pic_num);
+							
+							OutSlice.RefPicList0.Add({false, pic_num});	
+						}
+						else if (modification_of_pic_nums_idc == 2)
+						{
+							FNalu::UE pic_num;
+							Bitstream.Read(pic_num);
+							
+							OutSlice.RefPicList0.Add({true, pic_num});							
+						}
+					}
+					while (modification_of_pic_nums_idc != 3);
+				}
+			}
+
+			if (OutSlice.slice_type % 5 == 1)
+			{
+				Bitstream.Read(OutSlice.ref_pic_list_modification_flag_l1);
+				if (OutSlice.ref_pic_list_modification_flag_l1)
+				{
+					uint32 modification_of_pic_nums_idc;
+					do
+					{
+						Bitstream.Read(modification_of_pic_nums_idc);
+						if (modification_of_pic_nums_idc == 0 || modification_of_pic_nums_idc == 1)
+						{
+							FNalu::UE pic_num;
+							Bitstream.Read(pic_num);
+							
+							OutSlice.RefPicList1.Add({false, pic_num});	
+						}
+						else if (modification_of_pic_nums_idc == 2)
+						{
+							FNalu::UE pic_num;
+							Bitstream.Read(pic_num);
+							
+							OutSlice.RefPicList1.Add({true, pic_num});							
+						}
+					}
+					while (modification_of_pic_nums_idc != 3);
+				}
+			}
+		}
+
+		if ( (CurrentPPS.weighted_pred_flag && ( IsP || IsSP )) || (CurrentPPS.weighted_bipred_idc == 1 && IsB) )
+		{
+			// TODO pred_weight_table();
+		}
+
+		if (InNaluInfo.RefIdc != 0)
+		{
+			// TODO dec_ref_pic_marking()
+		}
+
+		if (CurrentPPS.entropy_coding_mode_flag && !IsI && !IsSI)
+		{
+			Bitstream.Read(OutSlice.cabac_init_idc);
+		}
+
+		Bitstream.Read(OutSlice.slice_qp_delta);
+
+		if ( IsSP || IsSI )
+		{
+			if (IsSP)
+			{
+				Bitstream.Read(OutSlice.sp_for_switch_flag);
+			}
+
+			Bitstream.Read(OutSlice.slice_qs_delta);
+		}
+
+		if (CurrentPPS.deblocking_filter_control_present_flag)
+		{
+			Bitstream.Read(OutSlice.disable_deblocking_filter_idc);
+
+			if (OutSlice.disable_deblocking_filter_idc != 1)
+			{
+				Bitstream.Read(OutSlice.slice_alpha_c0_offset_div2);
+				Bitstream.Read(OutSlice.slice_beta_offset_div2);
+			}
+		}
+
+		if (CurrentPPS.num_slice_groups_minus1 > 0
+			&& CurrentPPS.slice_group_map_type >=3
+			&& CurrentPPS.slice_group_map_type <= 5)
+		{
+			const uint32 NumBits = FMath::CeilLogTwo((CurrentPPS.pic_size_in_map_units_minus1 + 1) / ( CurrentPPS.slice_group_change_rate_minus1 + 2) );
+			Bitstream.ReadBits(OutSlice.slice_group_change_cycle, NumBits);
+		}
+		
+		return EAVResult::Success;
+	}
+
+} // namespace UE::AVCodecCore::H264
+
+// template <>
+// DLLEXPORT FAVResult FAVExtension::TransformPacket(FVideoPacketH264& OutPacket, FVideoPacket const& InPacket)
+// {
+// 	using namespace UE::AVCodecCore::H264;
+
+// 	TArray<FNaluInfo> FoundNalus;
+
+// 	FAVResult Result;
+// 	Result = FindNALUs(InPacket, FoundNalus);
+
+// 	if (Result.IsNotSuccess())
+// 	{
+// 		return Result;
+// 	}
+
+// 	for (auto& NaluInfo : FoundNalus)
+// 	{		
+// 		Result = EAVResult::Success;
+		
+// 		// NALUs are is usually an EBSP so we need to strip out the emulation prevention 3 byte making them RBSP
+// 		TArray64<uint8> RBSP;
+// 		RBSP.AddUninitialized(NaluInfo.Size);
+
+// 		int32 RBSPsize = EBSPtoRBSP(RBSP.GetData(), NaluInfo.Data, NaluInfo.Size);
+// 		FBitDataStream Bitstream(RBSP.GetData(), RBSPsize);
+
+// 		switch (NaluInfo.Type)
+// 		{
+// 			case ENaluType::Unspecified:
+// 			default:
+// 				Result = FAVResult(EAVResult::Error, TEXT("Recieved ENaluType::Unspecified"), TEXT("BitstreamParserH264"));
+// 				break;
+// 			case ENaluType::SliceOfNonIdrPicture:
+// 				OutPacket.Slices.AddUninitialized();
+// 				Result = ParseSliceHeader(Bitstream, NaluInfo, OutPacket.SPS, OutPacket.PPS, OutPacket.Slices.Last());
+// 				break;
+// 			case ENaluType::SliceDataPartitionA:
+// 				break;
+// 			case ENaluType::SliceDataPartitionB:
+// 				break;
+// 			case ENaluType::SliceDataPartitionC:
+// 				break;
+// 			case ENaluType::SliceIdrPicture:
+// 				break;
+// 			case ENaluType::SupplementalEnhancementInformation:
+// 				OutPacket.SEI.AddUninitialized();
+// 				Result = ParseSEI(Bitstream, NaluInfo, OutPacket.SEI.Last());
+// 				break;
+// 			case ENaluType::SequenceParameterSet:
+// 				Result = ParseSPS(Bitstream, NaluInfo, OutPacket.SPS);
+// 				break;
+// 			case ENaluType::PictureParameterSet:
+// 				Result = ParsePPS(Bitstream, NaluInfo, OutPacket.SPS, OutPacket.PPS);
+// 				break;
+// 			case ENaluType::AccessUnitDelimiter:
+// 				break;
+// 			case ENaluType::EndOfSequence:
+// 				break;
+// 			case ENaluType::EndOfStream:
+// 				break;
+// 			case ENaluType::FillerData:
+// 				break;
+// 			case ENaluType::SequenceParameterSetExtension:
+// 				break;
+// 			case ENaluType::PrefixNalUnit:
+// 				break;
+// 			case ENaluType::SubsetSequenceParameterSet:
+// 				break;
+// 			case ENaluType::Reserved16:
+// 				Result = FAVResult(EAVResult::Error, TEXT("Recieved ENaluType::Reserved16"), TEXT("BitstreamParserH264"));
+// 				break;
+// 			case ENaluType::Reserved17:
+// 				Result = FAVResult(EAVResult::Error, TEXT("Recieved ENaluType::Reserved17"), TEXT("BitstreamParserH264"));
+// 				break;
+// 			case ENaluType::Reserved18:
+// 				Result = FAVResult(EAVResult::Error, TEXT("Recieved ENaluType::Reserved18"), TEXT("BitstreamParserH264"));
+// 				break;
+// 			case ENaluType::SliceOfAnAuxiliaryCoded:
+// 				break;
+// 			case ENaluType::SliceExtension:
+// 				break;
+// 			case ENaluType::SliceExtensionForDepthView:
+// 				break;
+// 			case ENaluType::Reserved22:
+// 				Result = FAVResult(EAVResult::Error, TEXT("Recieved ENaluType::Reserved22"), TEXT("BitstreamParserH264"));
+// 				break;
+// 			case ENaluType::Reserved23:
+// 				Result = FAVResult(EAVResult::Error, TEXT("Recieved ENaluType::Reserved23"), TEXT("BitstreamParserH264"));
+// 				break;
+// 			case ENaluType::Unspecified24:
+// 				Result = FAVResult(EAVResult::Error, TEXT("Recieved ENaluType::Unspecified24"), TEXT("BitstreamParserH264"));
+// 				break;
+// 			case ENaluType::Unspecified25:
+// 				Result = FAVResult(EAVResult::Error, TEXT("Recieved ENaluType::Unspecified25"), TEXT("BitstreamParserH264"));
+// 				break;
+// 			case ENaluType::Unspecified26:
+// 				Result = FAVResult(EAVResult::Error, TEXT("Recieved ENaluType::Unspecified26"), TEXT("BitstreamParserH264"));
+// 				break;
+// 			case ENaluType::Unspecified27:
+// 				Result = FAVResult(EAVResult::Error, TEXT("Recieved ENaluType::Unspecified27"), TEXT("BitstreamParserH264"));
+// 				break;
+// 			case ENaluType::Unspecified28:
+// 				Result = FAVResult(EAVResult::Error, TEXT("Recieved ENaluType::Unspecified28"), TEXT("BitstreamParserH264"));
+// 				break;
+// 			case ENaluType::Unspecified29:
+// 				Result = FAVResult(EAVResult::Error, TEXT("Recieved ENaluType::Unspecified29"), TEXT("BitstreamParserH264"));
+// 				break;
+// 			case ENaluType::Unspecified30:
+// 				Result = FAVResult(EAVResult::Error, TEXT("Recieved ENaluType::Unspecified30"), TEXT("BitstreamParserH264"));
+// 				break;
+// 			case ENaluType::Unspecified31:
+// 				Result = FAVResult(EAVResult::Error, TEXT("Recieved ENaluType::Unspecified31"), TEXT("BitstreamParserH264"));
+// 				break;
+// 		}
+
+// 		if (Result.IsNotSuccess())
+// 		{
+// 			return Result;
+// 		}
+// 	}
+
+// 	return Result;
+// }
diff --git a/Engine/Plugins/Experimental/AVCodecs/AVCodecsCore/Source/AVCodecsCore/Private/Video/CodecUtils/CodecUtilsH265.cpp b/Engine/Plugins/Experimental/AVCodecs/AVCodecsCore/Source/AVCodecsCore/Private/Video/CodecUtils/CodecUtilsH265.cpp
new file mode 100644
--- /dev/null
+++ b/Engine/Plugins/Experimental/AVCodecs/AVCodecsCore/Source/AVCodecsCore/Private/Video/CodecUtils/CodecUtilsH265.cpp
@@ -0,0 +1,1650 @@
+// Copyright Epic Games, Inc. All Rights Reserved.
+
+#include "Video/CodecUtils/CodecUtilsH265.h"
+
+#include "Video/VideoPacket.h"
+#include "Video/Decoders/Configs/VideoDecoderConfigH265.h"
+
+FH265ProfileDefinition GH265ProfileDefinitions[static_cast<uint8>(EH265Profile::MAX)] = {
+	{ EH265Profile::Auto, UE::AVCodecCore::H265::EH265ProfileIDC::Auto, UE::AVCodecCore::H265::EH265ConstraintFlag::None, TEXT("Auto") }
+};
+
+namespace UE::AVCodecCore::H265
+{
+	FAVResult FindNALUs(FVideoPacket const& InPacket, TArray<FNaluH265>& FoundNalus)
+	{
+		FBitstreamReader Reader(InPacket.DataPtr.Get(), InPacket.DataSize);
+		if (Reader.NumBytesRemaining() < 3)
+		{
+			return FAVResult(EAVResult::Warning, TEXT("Bitstream not long enough to hold a NALU"), TEXT("H265"));
+		}
+
+		const TArrayView64<uint8> Data = InPacket.GetData();
+
+		// HACK if the packet does not start with a NAL delimiter we assume that the start of the stream is the first NAL
+		bool bHasLeadingZeros = (Data[0] == 0 && Data[1] == 0);
+		bool bIsNal3 = Data[2] == 1;
+		bool bIsNal4 = (Data[2] == 0 && Data[3] == 1);		
+		if (!(bHasLeadingZeros && (bIsNal3 || bIsNal4)))
+		{
+			FNaluH265 Nalu;
+			Nalu.StartIdx = 0;
+			Nalu.StartCodeSize = 0;
+			
+			// Extract NAL Header
+			{
+				FBitstreamReader Bitstream(&Data[Nalu.StartIdx + Nalu.StartCodeSize], 1);
+				verifyf(Bitstream.ReadBits(1) == 0, TEXT("Forbidden Zero bit not Zero in NAL Header"));
+
+				Nalu.nal_unit_type = (ENaluType)Bitstream.ReadBits(6);
+				Nalu.nuh_layer_id = Bitstream.ReadBits(6);
+				Nalu.nuh_temporal_id_plus1 = Bitstream.ReadBits(3);
+			}
+
+			Nalu.EBSP = &Data[Nalu.StartIdx];
+			
+			FoundNalus.Add(Nalu);
+		}
+		
+		// Skip over stream in intervals of 3 until Data[i + 2] is either 0 or 1
+		for (int64 i = 0; i < Data.Num() - 2;)
+		{
+			if (Data[i + 2] > 1)
+			{
+				i += 3;
+			}
+			else if (Data[i + 2] == 1)
+			{
+				if (Data[i + 1] == 0 && Data[i] == 0)
+				{
+					// Found start sequence of NalU but we don't know if it has a 3 or 4 byte start code so we check.
+					FNaluH265 Nalu;
+					Nalu.StartIdx = i;
+
+					if (Nalu.StartIdx > 0 && Data[Nalu.StartIdx - 1] == 0)
+					{
+						++Nalu.StartCodeSize;
+						--Nalu.StartIdx;
+					}
+
+					// Extract NAL Header
+					{
+						FBitstreamReader Bitstream(&Data[Nalu.StartIdx + Nalu.StartCodeSize], 1);
+						verifyf(Bitstream.ReadBits(1) == 0, TEXT("Forbidden Zero bit not Zero in NAL Header"));
+
+						Nalu.nal_unit_type = (ENaluType)Bitstream.ReadBits(6);
+						Nalu.nuh_layer_id = Bitstream.ReadBits(6);
+						Nalu.nuh_temporal_id_plus1 = Bitstream.ReadBits(3);
+					}
+
+					Nalu.EBSP = &Data[Nalu.StartIdx];
+
+					// Update length of previous entry.
+					if (FoundNalus.Num() > 0)
+					{
+						FoundNalus.Last().Size = Nalu.StartIdx - FoundNalus.Last().StartIdx;
+					}
+
+					FoundNalus.Add(Nalu);
+				}
+
+				i += 3;
+			}
+			else
+			{
+				++i;
+			}
+		}
+
+		if (FoundNalus.Num() == 0)
+		{
+			return FAVResult(EAVResult::PendingInput, TEXT("no NALUs found in BitDataStream"), TEXT("BitstreamParserH265"));
+		}
+
+		// Last Nal size is the remaining size of the bitstream minus a trailing zero byte
+		FoundNalus.Last().Size = Data.Num() - FoundNalus.Last().StartIdx;
+
+		// TODO (aidan) make a verbose msg
+		FAVResult::Log(EAVResult::Success, FString::Printf(TEXT("FAVExtension::TransformConfig found %d NALUs in bitdatastream"), FoundNalus.Num()));
+
+		// TODO (aidan) make a debug only verbose msg
+		for (const FNaluH265& NalUInfo : FoundNalus)
+		{
+			FAVResult::Log(EAVResult::Success, FString::Printf(TEXT("Found NALU at %llu size %llu with type %u"), NalUInfo.StartIdx, NalUInfo.Size, (uint8)NalUInfo.nal_unit_type.Value));
+		}
+
+		return EAVResult::Success;
+	}
+
+	void profile_tier_level_t::Parse(const uint8& profilePresentFlag, uint8 maxNumSubLayersMinus1, FBitstreamReader& Bitstream)
+	{
+		if (profilePresentFlag)
+		{
+			Bitstream.Read(
+				general_profile_space,
+				general_tier_flag,
+				general_profile_idc,
+				general_profile_compatibility_flag,
+				general_progressive_source_flag,
+				general_interlaced_source_flag);
+
+			general_constraint_flags |= Bitstream.ReadBits(1) ? EH265ConstraintFlag::non_packed_constraint_flag : EH265ConstraintFlag::None;
+			general_constraint_flags |= Bitstream.ReadBits(1) ? EH265ConstraintFlag::frame_only_constraint_flag : EH265ConstraintFlag::None;
+
+			if (
+				general_profile_idc == EH265ProfileIDC::FormatRangeExtensions || CheckProfileCompatabilityFlag(general_profile_compatibility_flag, EH265ProfileIDC::FormatRangeExtensions) || general_profile_idc == EH265ProfileIDC::HighThroughput || CheckProfileCompatabilityFlag(general_profile_compatibility_flag, EH265ProfileIDC::HighThroughput) ||
+				// general_profile_idc == 6 || CheckProfileCompatabilityFlag(general_profile_compatibility_flag, 6 ) || // TODO Not sure what this profile idc is
+				// general_profile_idc == 7 || CheckProfileCompatabilityFlag(general_profile_compatibility_flag, 7 ) || // TODO Not sure what this profile idc is
+				// general_profile_idc == 8 || CheckProfileCompatabilityFlag(general_profile_compatibility_flag, 8 ) || // TODO Not sure what this profile idc is
+				general_profile_idc == EH265ProfileIDC::ScreenContentCoding || CheckProfileCompatabilityFlag(general_profile_compatibility_flag, EH265ProfileIDC::ScreenContentCoding) ||
+				// general_profile_idc == 10 || CheckProfileCompatabilityFlag(general_profile_compatibility_flag, 10 ) || // TODO Not sure what this profile idc is
+				general_profile_idc == EH265ProfileIDC::HighThroughputScreenContentCoding || CheckProfileCompatabilityFlag(general_profile_compatibility_flag, EH265ProfileIDC::HighThroughputScreenContentCoding))
+			{
+				general_constraint_flags |= Bitstream.ReadBits(1) ? EH265ConstraintFlag::max_12bit_constraint_flag : EH265ConstraintFlag::None;
+				general_constraint_flags |= Bitstream.ReadBits(1) ? EH265ConstraintFlag::max_10bit_constraint_flag : EH265ConstraintFlag::None;
+				general_constraint_flags |= Bitstream.ReadBits(1) ? EH265ConstraintFlag::max_8bit_constraint_flag : EH265ConstraintFlag::None;
+				general_constraint_flags |= Bitstream.ReadBits(1) ? EH265ConstraintFlag::max_422chroma_constraint_flag : EH265ConstraintFlag::None;
+				general_constraint_flags |= Bitstream.ReadBits(1) ? EH265ConstraintFlag::max_420chroma_constraint_flag : EH265ConstraintFlag::None;
+				general_constraint_flags |= Bitstream.ReadBits(1) ? EH265ConstraintFlag::max_monochrome_constraint_flag : EH265ConstraintFlag::None;
+				general_constraint_flags |= Bitstream.ReadBits(1) ? EH265ConstraintFlag::intra_constraint_flag : EH265ConstraintFlag::None;
+				general_constraint_flags |= Bitstream.ReadBits(1) ? EH265ConstraintFlag::one_picture_only_constraint_flag : EH265ConstraintFlag::None;
+				general_constraint_flags |= Bitstream.ReadBits(1) ? EH265ConstraintFlag::lower_bit_rate_constraint_flag : EH265ConstraintFlag::None;
+
+				if (
+					general_profile_idc == EH265ProfileIDC::HighThroughput || CheckProfileCompatabilityFlag(general_profile_compatibility_flag, EH265ProfileIDC::HighThroughput) || general_profile_idc == EH265ProfileIDC::ScreenContentCoding || CheckProfileCompatabilityFlag(general_profile_compatibility_flag, EH265ProfileIDC::ScreenContentCoding) ||
+					// general_profile_idc == 10 || CheckProfileCompatabilityFlag(general_profile_compatibility_flag, 10 ) || // TODO Not sure what this profile idc is
+					general_profile_idc == EH265ProfileIDC::HighThroughputScreenContentCoding || CheckProfileCompatabilityFlag(general_profile_compatibility_flag, EH265ProfileIDC::HighThroughputScreenContentCoding))
+				{
+					general_constraint_flags |= Bitstream.ReadBits(1) ? EH265ConstraintFlag::max_14bit_constraint_flag : EH265ConstraintFlag::None;
+					Bitstream.SkipBits(33);
+				}
+				else
+				{
+					Bitstream.SkipBits(34);
+				}
+			}
+			else if (general_profile_idc == EH265ProfileIDC::Main10 || CheckProfileCompatabilityFlag(general_profile_compatibility_flag, EH265ProfileIDC::Main10))
+			{
+				Bitstream.SkipBits(7);
+				general_constraint_flags |= Bitstream.ReadBits(1) ? EH265ConstraintFlag::one_picture_only_constraint_flag : EH265ConstraintFlag::None;
+				Bitstream.SkipBits(35);
+			}
+			else
+			{
+				Bitstream.SkipBits(43);
+			}
+
+			if (
+				general_profile_idc == EH265ProfileIDC::Main || CheckProfileCompatabilityFlag(general_profile_compatibility_flag, EH265ProfileIDC::Main) || general_profile_idc == EH265ProfileIDC::Main10 || CheckProfileCompatabilityFlag(general_profile_compatibility_flag, EH265ProfileIDC::Main10) || general_profile_idc == EH265ProfileIDC::MainStillPicture || CheckProfileCompatabilityFlag(general_profile_compatibility_flag, EH265ProfileIDC::MainStillPicture) || general_profile_idc == EH265ProfileIDC::FormatRangeExtensions || CheckProfileCompatabilityFlag(general_profile_compatibility_flag, EH265ProfileIDC::FormatRangeExtensions) || general_profile_idc == EH265ProfileIDC::HighThroughput || CheckProfileCompatabilityFlag(general_profile_compatibility_flag, EH265ProfileIDC::HighThroughput) || general_profile_idc == EH265ProfileIDC::ScreenContentCoding || CheckProfileCompatabilityFlag(general_profile_compatibility_flag, EH265ProfileIDC::ScreenContentCoding) || general_profile_idc == EH265ProfileIDC::HighThroughputScreenContentCoding || CheckProfileCompatabilityFlag(general_profile_compatibility_flag, EH265ProfileIDC::HighThroughputScreenContentCoding))
+			{
+				Bitstream.Read(general_inbld_flag);
+			}
+			else
+			{
+				Bitstream.SkipBits(1);
+			}
+
+			Bitstream.Read(general_level_idc);
+
+			sub_layers.SetNumUninitialized(maxNumSubLayersMinus1 + 1);
+			for (uint8 i = 0; i < maxNumSubLayersMinus1; i++)
+			{
+				Bitstream.Read(
+					sub_layers[i].sub_layer_profile_present_flag,
+					sub_layers[i].sub_layer_level_present_flag);
+			}
+
+			if (maxNumSubLayersMinus1 > 0)
+			{
+				for (uint8 i = maxNumSubLayersMinus1; i < 8; i++)
+				{
+					Bitstream.SkipBits(2);
+				}
+			}
+
+			for (uint8 i = 0; i < maxNumSubLayersMinus1; i++)
+			{
+				if (sub_layers[i].sub_layer_profile_present_flag)
+				{
+					Bitstream.Read(
+						sub_layers[i].sub_layer_profile_space,
+						sub_layers[i].sub_layer_tier_flag,
+						sub_layers[i].sub_layer_profile_idc,
+						sub_layers[i].sub_layer_profile_compatibility_flag,
+						sub_layers[i].sub_layer_progressive_source_flag,
+						sub_layers[i].sub_layer_interlaced_source_flag);
+
+					sub_layers[i].sub_layer_constraint_flags |= Bitstream.ReadBits(1) ? EH265ConstraintFlag::non_packed_constraint_flag : EH265ConstraintFlag::None;
+					sub_layers[i].sub_layer_constraint_flags |= Bitstream.ReadBits(1) ? EH265ConstraintFlag::frame_only_constraint_flag : EH265ConstraintFlag::None;
+
+					if (
+						sub_layers[i].sub_layer_profile_idc == EH265ProfileIDC::FormatRangeExtensions || CheckProfileCompatabilityFlag(sub_layers[i].sub_layer_profile_compatibility_flag, EH265ProfileIDC::FormatRangeExtensions) || sub_layers[i].sub_layer_profile_idc == EH265ProfileIDC::HighThroughput || CheckProfileCompatabilityFlag(sub_layers[i].sub_layer_profile_compatibility_flag, EH265ProfileIDC::HighThroughput) ||
+						// sub_layers[i].sub_layer_profile_idc == 6 || CheckProfileCompatabilityFlag(sub_layers[i].sub_layer_profile_compatibility_flag, 6 ) || // TODO Not sure what this profile idc is
+						// sub_layers[i].sub_layer_profile_idc == 7 || CheckProfileCompatabilityFlag(sub_layers[i].sub_layer_profile_compatibility_flag, 7 ) || // TODO Not sure what this profile idc is
+						// sub_layers[i].sub_layer_profile_idc == 8 || CheckProfileCompatabilityFlag(sub_layers[i].sub_layer_profile_compatibility_flag, 8 ) || // TODO Not sure what this profile idc is
+						sub_layers[i].sub_layer_profile_idc == EH265ProfileIDC::ScreenContentCoding || CheckProfileCompatabilityFlag(sub_layers[i].sub_layer_profile_compatibility_flag, EH265ProfileIDC::ScreenContentCoding) ||
+						// sub_layers[i].sub_layer_profile_idc == 10 || CheckProfileCompatabilityFlag(sub_layers[i].sub_layer_profile_compatibility_flag, 10 ) || // TODO Not sure what this profile idc is
+						sub_layers[i].sub_layer_profile_idc == EH265ProfileIDC::HighThroughputScreenContentCoding || CheckProfileCompatabilityFlag(sub_layers[i].sub_layer_profile_compatibility_flag, EH265ProfileIDC::HighThroughputScreenContentCoding))
+					{
+						sub_layers[i].sub_layer_constraint_flags |= Bitstream.ReadBits(1) ? EH265ConstraintFlag::max_12bit_constraint_flag : EH265ConstraintFlag::None;
+						sub_layers[i].sub_layer_constraint_flags |= Bitstream.ReadBits(1) ? EH265ConstraintFlag::max_10bit_constraint_flag : EH265ConstraintFlag::None;
+						sub_layers[i].sub_layer_constraint_flags |= Bitstream.ReadBits(1) ? EH265ConstraintFlag::max_8bit_constraint_flag : EH265ConstraintFlag::None;
+						sub_layers[i].sub_layer_constraint_flags |= Bitstream.ReadBits(1) ? EH265ConstraintFlag::max_422chroma_constraint_flag : EH265ConstraintFlag::None;
+						sub_layers[i].sub_layer_constraint_flags |= Bitstream.ReadBits(1) ? EH265ConstraintFlag::max_420chroma_constraint_flag : EH265ConstraintFlag::None;
+						sub_layers[i].sub_layer_constraint_flags |= Bitstream.ReadBits(1) ? EH265ConstraintFlag::max_monochrome_constraint_flag : EH265ConstraintFlag::None;
+						sub_layers[i].sub_layer_constraint_flags |= Bitstream.ReadBits(1) ? EH265ConstraintFlag::intra_constraint_flag : EH265ConstraintFlag::None;
+						sub_layers[i].sub_layer_constraint_flags |= Bitstream.ReadBits(1) ? EH265ConstraintFlag::one_picture_only_constraint_flag : EH265ConstraintFlag::None;
+						sub_layers[i].sub_layer_constraint_flags |= Bitstream.ReadBits(1) ? EH265ConstraintFlag::lower_bit_rate_constraint_flag : EH265ConstraintFlag::None;
+
+						if (
+							sub_layers[i].sub_layer_profile_idc == EH265ProfileIDC::HighThroughput || CheckProfileCompatabilityFlag(sub_layers[i].sub_layer_profile_compatibility_flag, EH265ProfileIDC::HighThroughput) || sub_layers[i].sub_layer_profile_idc == EH265ProfileIDC::ScreenContentCoding || CheckProfileCompatabilityFlag(sub_layers[i].sub_layer_profile_compatibility_flag, EH265ProfileIDC::ScreenContentCoding) ||
+							// sub_layers[i].sub_layer_profile_idc == 10 || CheckProfileCompatabilityFlag(sub_layers[i].sub_layer_profile_compatibility_flag, 10 ) || // TODO Not sure what this profile idc is
+							sub_layers[i].sub_layer_profile_idc == EH265ProfileIDC::HighThroughputScreenContentCoding || CheckProfileCompatabilityFlag(sub_layers[i].sub_layer_profile_compatibility_flag, EH265ProfileIDC::HighThroughputScreenContentCoding))
+						{
+							sub_layers[i].sub_layer_constraint_flags |= Bitstream.ReadBits(1) ? EH265ConstraintFlag::max_14bit_constraint_flag : EH265ConstraintFlag::None;
+							Bitstream.SkipBits(33);
+						}
+						else
+						{
+							Bitstream.SkipBits(34);
+						}
+					}
+					else if (sub_layers[i].sub_layer_profile_idc == EH265ProfileIDC::Main10 || CheckProfileCompatabilityFlag(sub_layers[i].sub_layer_profile_compatibility_flag, EH265ProfileIDC::Main10))
+					{
+						Bitstream.SkipBits(7);
+						sub_layers[i].sub_layer_constraint_flags |= Bitstream.ReadBits(1) ? EH265ConstraintFlag::one_picture_only_constraint_flag : EH265ConstraintFlag::None;
+						Bitstream.SkipBits(35);
+					}
+					else
+					{
+						Bitstream.SkipBits(43);
+					}
+
+					if (
+						sub_layers[i].sub_layer_profile_idc == EH265ProfileIDC::Main || CheckProfileCompatabilityFlag(sub_layers[i].sub_layer_profile_compatibility_flag, EH265ProfileIDC::Main) || sub_layers[i].sub_layer_profile_idc == EH265ProfileIDC::Main10 || CheckProfileCompatabilityFlag(sub_layers[i].sub_layer_profile_compatibility_flag, EH265ProfileIDC::Main10) || sub_layers[i].sub_layer_profile_idc == EH265ProfileIDC::MainStillPicture || CheckProfileCompatabilityFlag(sub_layers[i].sub_layer_profile_compatibility_flag, EH265ProfileIDC::MainStillPicture) || sub_layers[i].sub_layer_profile_idc == EH265ProfileIDC::FormatRangeExtensions || CheckProfileCompatabilityFlag(sub_layers[i].sub_layer_profile_compatibility_flag, EH265ProfileIDC::FormatRangeExtensions) || sub_layers[i].sub_layer_profile_idc == EH265ProfileIDC::HighThroughput || CheckProfileCompatabilityFlag(sub_layers[i].sub_layer_profile_compatibility_flag, EH265ProfileIDC::HighThroughput) || sub_layers[i].sub_layer_profile_idc == EH265ProfileIDC::ScreenContentCoding || CheckProfileCompatabilityFlag(sub_layers[i].sub_layer_profile_compatibility_flag, EH265ProfileIDC::ScreenContentCoding) || sub_layers[i].sub_layer_profile_idc == EH265ProfileIDC::HighThroughputScreenContentCoding || CheckProfileCompatabilityFlag(sub_layers[i].sub_layer_profile_compatibility_flag, EH265ProfileIDC::HighThroughputScreenContentCoding))
+					{
+						Bitstream.Read(sub_layers[i].sub_layer_inbld_flag);
+					}
+					else
+					{
+						Bitstream.SkipBits(1);
+					}
+				}
+
+				if (sub_layers[i].sub_layer_level_present_flag)
+				{
+					Bitstream.Read(sub_layers[i].sub_layer_level_idc);
+				}
+			}
+		}
+	}
+
+	void hrd_parameters_t::sub_layer_t::Parse(uint8 const& in_sub_pic_hrd_params_present_flag, uint8 const& CpbCnt, FBitstreamReader& Bitstream)
+	{
+		sub_layer_hrd_parameters.SetNumZeroed(CpbCnt);
+
+		for (uint8 i = 0; i < CpbCnt; i++)
+		{
+			Bitstream.Read(
+				sub_layer_hrd_parameters[i].bit_rate_value_minus1,
+				sub_layer_hrd_parameters[i].cpb_size_value_minus1);
+			if (in_sub_pic_hrd_params_present_flag)
+			{
+				Bitstream.Read(
+					sub_layer_hrd_parameters[i].cpb_size_du_value_minus1,
+					sub_layer_hrd_parameters[i].bit_rate_du_value_minus1);
+			}
+			Bitstream.Read(sub_layer_hrd_parameters[i].cbr_flag);
+		}
+	}
+
+	void hrd_parameters_t::Parse(uint8 const& commonInfPresentFlag, uint8 const& maxNumSubLayersMinus1, FBitstreamReader& Bitstream)
+	{
+		if (commonInfPresentFlag)
+		{
+			Bitstream.Read(
+				nal_hrd_parameters_present_flag,
+				vcl_hrd_parameters_present_flag);
+
+			if (nal_hrd_parameters_present_flag || vcl_hrd_parameters_present_flag)
+			{
+				Bitstream.Read(sub_pic_hrd_params_present_flag);
+				if (sub_pic_hrd_params_present_flag)
+				{
+					Bitstream.Read(
+						tick_divisor_minus2,
+						du_cpb_removal_delay_increment_length_minus1,
+						sub_pic_cpb_params_in_pic_timing_sei_flag,
+						dpb_output_delay_du_length_minus1);
+				}
+
+				Bitstream.Read(
+					bit_rate_scale,
+					cpb_size_scale);
+
+				if (sub_pic_hrd_params_present_flag)
+				{
+					Bitstream.Read(cpb_size_du_scale);
+				}
+
+				Bitstream.Read(
+					initial_cpb_removal_delay_length_minus1,
+					au_cpb_removal_delay_length_minus1,
+					dpb_output_delay_length_minus1);
+			}
+		}
+
+		sub_layers.SetNumZeroed(maxNumSubLayersMinus1 + 1);
+
+		for (uint8 i = 0; i <= maxNumSubLayersMinus1; i++)
+		{
+			Bitstream.Read(sub_layers[i].fixed_pic_rate_general_flag);
+			if (!sub_layers[i].fixed_pic_rate_general_flag)
+			{
+				Bitstream.Read(sub_layers[i].fixed_pic_rate_within_cvs_flag);
+			}
+
+			if (sub_layers[i].fixed_pic_rate_within_cvs_flag)
+			{
+				Bitstream.Read(sub_layers[i].elemental_duration_in_tc_minus1);
+			}
+			else
+			{
+				Bitstream.Read(sub_layers[i].low_delay_hrd_flag);
+			}
+
+			if (!sub_layers[i].low_delay_hrd_flag)
+			{
+				Bitstream.Read(sub_layers[i].cpb_cnt_minus1);
+			}
+
+			if (nal_hrd_parameters_present_flag)
+			{
+				sub_layers[i].Parse(sub_pic_hrd_params_present_flag, sub_layers[i].cpb_cnt_minus1 + 1, Bitstream);
+			}
+
+			if (vcl_hrd_parameters_present_flag)
+			{
+				sub_layers[i].Parse(sub_pic_hrd_params_present_flag, sub_layers[i].cpb_cnt_minus1 + 1, Bitstream);
+			}
+		}
+	}
+
+	FAVResult FNaluVPS::Parse()
+	{
+		// Get Bitstream
+		TArray64<uint8> RBSP;
+		RBSP.SetNumUninitialized(Size - StartCodeSize);
+
+		// Skip start code
+		int32 RBSPsize = EBSPtoRBSP(RBSP.GetData(), EBSP + StartCodeSize, Size - StartCodeSize);
+		FBitstreamReader Bitstream(RBSP.GetData(), RBSPsize);
+
+		// Skip NAL Header which we have already parsed 
+		Bitstream.SkipBytes(2);
+
+		// Start parsing VPS
+		Bitstream.Read(
+			vps_video_parameter_set_id,
+			vps_base_layer_internal_flag,
+			vps_base_layer_available_flag,
+			vps_max_layers_minus1,
+			vps_max_sub_layers_minus1,
+			vps_temporal_id_nesting_flag);
+
+		Bitstream.SkipBits(16); // Reserved
+
+		profile_tier_level.Parse(1, vps_max_layers_minus1, Bitstream);
+
+		vps_max_dec_pic_buffering_minus1.SetNumUninitialized(vps_max_sub_layers_minus1 + 1);
+		vps_max_num_reorder_pics.SetNumUninitialized(vps_max_sub_layers_minus1 + 1);
+		vps_max_latency_increase_plus1.SetNumUninitialized(vps_max_sub_layers_minus1 + 1);
+
+		Bitstream.Read(vps_sub_layer_ordering_info_present_flag);
+		for (uint8 i = (vps_sub_layer_ordering_info_present_flag ? 0u : (uint8)vps_max_sub_layers_minus1); i <= vps_max_sub_layers_minus1; i++)
+		{
+			Bitstream.Read(
+				vps_max_dec_pic_buffering_minus1[i],
+				vps_max_num_reorder_pics[i],
+				vps_max_latency_increase_plus1[i]);
+		}
+
+		Bitstream.Read(
+			vps_max_layer_id,
+			vps_num_layer_sets_minus1);
+
+		layer_id_included_flag.SetNumZeroed(vps_num_layer_sets_minus1 + 1);
+		for (uint8 i = 1; i <= vps_num_layer_sets_minus1; i++)
+		{
+			layer_id_included_flag[i].SetNumUninitialized(vps_max_layer_id + 1);
+			for (uint8 j = 0; j <= vps_max_layer_id; j++)
+			{
+				Bitstream.Read(layer_id_included_flag[i][j]);
+			}
+		}
+
+		Bitstream.Read(vps_timing_info_present_flag);
+		if (vps_timing_info_present_flag)
+		{
+			Bitstream.Read(
+				vps_num_units_in_tick,
+				vps_time_scale);
+
+			vps_poc_proportional_to_timing_flag = Bitstream.PeekBits(1);
+			if (vps_poc_proportional_to_timing_flag)
+			{
+				Bitstream.Read(vps_num_ticks_poc_diff_one_minus1);
+			}
+
+			Bitstream.Read(vps_num_hrd_parameters);
+
+			hrd_layer_set_idx.SetNumZeroed(vps_num_hrd_parameters);
+			cprms_present_flag.SetNumZeroed(vps_num_hrd_parameters);
+			HrdParameters.SetNumZeroed(vps_num_hrd_parameters);
+			for (uint32 i = 0; i < vps_num_hrd_parameters; i++)
+			{
+				Bitstream.Read(hrd_layer_set_idx[i]);
+				if (i > 0)
+				{
+					Bitstream.Read(cprms_present_flag[i]);
+				}
+				HrdParameters[i].Parse(cprms_present_flag[i], vps_max_sub_layers_minus1, Bitstream);
+			}
+		}
+		Bitstream.Read(vps_extension_flag);
+		if (vps_extension_flag)
+		{
+			// TODO (aidan) handle VPS extentions ignoring for now
+		}
+
+		return EAVResult::Success;
+	}
+
+	void scaling_list_data_t::Parse(FBitstreamReader& Bitstream)
+	{
+		SE scaling_list_delta_coef;
+		for (uint8 sizeId = 0; sizeId < 4; sizeId++)
+		{
+			for (uint8 matrixId = 0; matrixId < 6; matrixId += (sizeId == 3) ? 3 : 1)
+			{
+				Bitstream.Read(scaling_list_pred_mode_flag[sizeId][matrixId]);
+				if (!scaling_list_pred_mode_flag[sizeId][matrixId])
+				{
+					Bitstream.Read(scaling_list_pred_matrix_id_delta[sizeId][matrixId]);
+				}
+				else
+				{
+					int16 nextCoef = 8;
+					uint8 coefNum = FMath::Min<uint8>(64u, (1 << (4 + (sizeId << 1))));
+
+					if (sizeId > 1)
+					{
+						Bitstream.Read(scaling_list_dc_coef_minus8[sizeId - 2][matrixId]);
+						nextCoef = scaling_list_dc_coef_minus8[sizeId - 2][matrixId] + 8;
+					}
+
+					if (sizeId == 0)
+					{
+						FMemory::Memset(ScalingList0[matrixId], 0, sizeof(ScalingList0[matrixId]));
+						for (uint8 i = 0; i < coefNum; i++)
+						{
+							Bitstream.Read(scaling_list_delta_coef);
+							nextCoef = (nextCoef + scaling_list_delta_coef + 256) % 256;
+							ScalingList0[matrixId][i] = nextCoef;
+						}
+					}
+					else
+					{
+						FMemory::Memset(ScalingList1to3[sizeId - 1], 0, sizeof(ScalingList1to3[sizeId - 1][matrixId]));
+						for (uint8 i = 0; i < coefNum; i++)
+						{
+							Bitstream.Read(scaling_list_delta_coef);
+							nextCoef = (nextCoef + scaling_list_delta_coef + 256) % 256;
+							ScalingList1to3[sizeId - 1][matrixId][i] = nextCoef;
+						}
+					}
+				}
+			}
+		}
+	}
+
+	void short_term_ref_pic_set_t::Parse(uint8 const& stRpsIdx, TArray<short_term_ref_pic_set_t> const& short_term_ref_pic_sets, FBitstreamReader& Bitstream)
+	{
+		if (stRpsIdx != 0)
+		{
+			Bitstream.Read(inter_ref_pic_set_prediction_flag);
+		}
+
+		if (inter_ref_pic_set_prediction_flag == 1)
+		{
+			if (stRpsIdx == short_term_ref_pic_sets.Num())
+			{
+				Bitstream.Read(delta_idx_minus1);
+			}
+
+			Bitstream.Read(
+				delta_rps_sign,
+				abs_delta_rps_minus1);
+
+			uint32 RefRpsIdx = stRpsIdx - (delta_idx_minus1 + 1);
+			uint8 RefNumDeltaPocs = short_term_ref_pic_sets[RefRpsIdx].NumDeltaPocs;
+
+			used_by_curr_pic_flags.SetNumUninitialized(RefNumDeltaPocs);
+			use_delta_flags.SetNumUninitialized(RefNumDeltaPocs);
+
+			for (uint8 j = 0; j <= RefNumDeltaPocs; j++)
+			{
+				Bitstream.Read(used_by_curr_pic_flags[j]);
+				if (!used_by_curr_pic_flags[j])
+				{
+					Bitstream.Read(use_delta_flags[j]);
+				}
+				else
+				{
+					use_delta_flags[j] = 1; // when not present inferred as 1 according to (7.4.8)
+				}
+			}
+		}
+		else // inter_ref_pic_set_prediction_flag != 1
+		{
+			Bitstream.Read(
+				num_negative_pics,
+				num_positive_pics);
+
+			delta_poc_s0_minus1s.SetNumUninitialized(num_negative_pics);
+			used_by_curr_pic_s0_flags.SetNumUninitialized(num_negative_pics);
+			for (uint8 i = 0; i < num_negative_pics; i++)
+			{
+				Bitstream.Read(
+					delta_poc_s0_minus1s[i],
+					used_by_curr_pic_s0_flags[i]);
+			}
+
+			delta_poc_s1_minus1s.SetNumUninitialized(num_positive_pics);
+			used_by_curr_pic_s1_flags.SetNumUninitialized(num_positive_pics);
+			for (uint8 i = 0; i < num_positive_pics; i++)
+			{
+				Bitstream.Read(
+					delta_poc_s1_minus1s[i],
+					used_by_curr_pic_s1_flags[i]);
+			}
+		}
+	}
+
+	void short_term_ref_pic_set_t::CalculateValues(uint8 const& stRpsIdx, TArray<short_term_ref_pic_set_t>& st_ref_pic_set)
+	{
+		if (inter_ref_pic_set_prediction_flag)
+		{
+			uint8 RefRpsIdx = stRpsIdx - (delta_idx_minus1 + 1);
+			int16 deltaRps = (1 - 2 * delta_rps_sign) * (abs_delta_rps_minus1 + 1);
+			int16 dPoc;
+
+			uint8 i = 0;
+			for (int16 j = st_ref_pic_set[RefRpsIdx].NumPositivePics - 1; j >= 0; j--)
+			{
+				dPoc = st_ref_pic_set[RefRpsIdx].DeltaPocS1[j] + deltaRps;
+				if (dPoc < 0 && use_delta_flags[st_ref_pic_set[RefRpsIdx].NumNegativePics + j])
+				{
+					DeltaPocS0[i] = dPoc;
+					UsedByCurrPicS0[i++] = used_by_curr_pic_flags[st_ref_pic_set[RefRpsIdx].NumNegativePics + j];
+				}
+			}
+
+			if (deltaRps < 0 && use_delta_flags[st_ref_pic_set[RefRpsIdx].NumDeltaPocs])
+			{
+				DeltaPocS0[i] = deltaRps;
+				UsedByCurrPicS0[i++] = used_by_curr_pic_flags[st_ref_pic_set[RefRpsIdx].NumDeltaPocs];
+			}
+
+			for (int16 j = 0; j < st_ref_pic_set[RefRpsIdx].NumNegativePics; j++)
+			{
+				dPoc = st_ref_pic_set[RefRpsIdx].DeltaPocS0[j] + deltaRps;
+				if (dPoc < 0 && use_delta_flags[j])
+				{
+					DeltaPocS0[i] = dPoc;
+					UsedByCurrPicS0[i++] = used_by_curr_pic_flags[j];
+				}
+			}
+			NumNegativePics = i;
+
+			i = 0;
+			for (int16 j = st_ref_pic_set[RefRpsIdx].NumNegativePics - 1; j >= 0; j--)
+			{
+				dPoc = st_ref_pic_set[RefRpsIdx].DeltaPocS0[j] + deltaRps;
+				if (dPoc > 0 && use_delta_flags[j])
+				{
+					DeltaPocS1[i] = dPoc;
+					UsedByCurrPicS1[i++] = used_by_curr_pic_flags[j];
+				}
+			}
+
+			if (deltaRps > 0 && use_delta_flags[st_ref_pic_set[RefRpsIdx].NumDeltaPocs])
+			{
+				DeltaPocS1[i] = deltaRps;
+				UsedByCurrPicS1[i++] = used_by_curr_pic_flags[st_ref_pic_set[RefRpsIdx].NumDeltaPocs];
+			}
+
+			for (int16 j = 0; j < st_ref_pic_set[RefRpsIdx].NumPositivePics; j++)
+			{
+				dPoc = st_ref_pic_set[RefRpsIdx].DeltaPocS1[j] + deltaRps;
+				if (dPoc > 0 && use_delta_flags[st_ref_pic_set[RefRpsIdx].NumNegativePics + j])
+				{
+					DeltaPocS1[i] = dPoc;
+					UsedByCurrPicS1[i++] = used_by_curr_pic_flags[st_ref_pic_set[RefRpsIdx].NumNegativePics + j];
+				}
+			}
+			NumPositivePics = i;
+		}
+		else
+		{
+			NumNegativePics = num_negative_pics;
+			for (uint8 i = 0; i < NumNegativePics; i++)
+			{
+				UsedByCurrPicS0[i] = used_by_curr_pic_s0_flags[i];
+				DeltaPocS0[i] = i == 0 ? 0 : DeltaPocS0[i-1];
+				DeltaPocS0[i] -= ( delta_poc_s0_minus1s[i] + 1);
+			}
+
+			NumPositivePics = num_positive_pics;
+			for (uint8 i = 0; i < NumPositivePics; i++)
+			{
+				UsedByCurrPicS0[i] = used_by_curr_pic_s1_flags[i];
+				DeltaPocS1[i] = i == 0 ? 0 : DeltaPocS1[i-1];
+				DeltaPocS1[i] += ( delta_poc_s1_minus1s[i] + 1);
+			}
+
+			NumDeltaPocs = NumNegativePics + NumPositivePics;
+		}
+	}
+
+	void long_term_ref_pics_t::Parse(const uint8& log2_max_pic_order_cnt_lsb_minus4, FBitstreamReader& Bitstream)
+	{
+		Bitstream.ReadBits(lt_ref_pic_poc_lsb_sps, log2_max_pic_order_cnt_lsb_minus4 + 4);
+		Bitstream.Read(used_by_curr_pic_lt_sps_flag);
+	}
+
+	void FNaluSPS::vui_parameters_t::Parse(FBitstreamReader& Bitstream)
+	{
+		Bitstream.Read(aspect_ratio_info_present_flag);
+		if (aspect_ratio_info_present_flag)
+		{
+			Bitstream.Read(aspect_ratio_idc);
+			if (aspect_ratio_idc == 255)
+			{
+				Bitstream.Read(
+					sar_width,
+					sar_height);
+			}
+		}
+
+		Bitstream.Read(overscan_info_present_flag);
+		if (overscan_info_present_flag)
+		{
+			Bitstream.Read(overscan_appropriate_flag);
+		}
+
+		Bitstream.Read(video_signal_type_present_flag);
+		if (video_signal_type_present_flag)
+		{
+			Bitstream.Read(
+				video_format,
+				video_full_range_flag,
+				colour_description_present_flag);
+
+			if (colour_description_present_flag)
+			{
+				Bitstream.Read(
+					colour_primaries,
+					transfer_characteristics,
+					matrix_coeffs);
+			}
+		}
+
+		Bitstream.Read(chroma_loc_info_present_flag);
+		if (chroma_loc_info_present_flag)
+		{
+			Bitstream.Read(
+				chroma_sample_loc_type_top_field,
+				chroma_sample_loc_type_bottom_field);
+		}
+
+		Bitstream.Read(
+			neutral_chroma_indication_flag,
+			field_seq_flag,
+			frame_field_info_present_flag,
+			default_display_window_flag);
+		if (default_display_window_flag)
+		{
+			Bitstream.Read(
+				def_disp_win_left_offset,
+				def_disp_win_right_offset,
+				def_disp_win_top_offset,
+				def_disp_win_bottom_offset);
+		}
+
+		Bitstream.Read(vui_timing_info_present_flag);
+		if (vui_timing_info_present_flag)
+		{
+			Bitstream.Read(
+				vui_num_units_in_tick,
+				vui_time_scale,
+				vui_poc_proportional_to_timing_flag);
+
+			if (vui_poc_proportional_to_timing_flag)
+			{
+				Bitstream.Read(vui_num_ticks_poc_diff_one_minus1);
+			}
+
+			Bitstream.Read(vui_hrd_parameters_present_flag);
+			if (vui_hrd_parameters_present_flag)
+			{
+				unimplemented();
+				// hrd_parameters.Parse( 1, sps_max_sub_layers_minus1, Bitstream );
+			}
+		}
+
+		Bitstream.Read(bitstream_restriction_flag);
+		if (bitstream_restriction_flag)
+		{
+			Bitstream.Read(
+				tiles_fixed_structure_flag,
+				motion_vectors_over_pic_boundaries_flag,
+				restricted_ref_pic_lists_flag,
+				min_spatial_segmentation_idc,
+				max_bytes_per_pic_denom,
+				max_bits_per_min_cu_denom,
+				log2_max_mv_length_horizontal,
+				log2_max_mv_length_vertical);
+		}
+	}
+
+	void FNaluSPS::sps_range_extension_t::Parse(FBitstreamReader& Bitstream)
+	{
+		Bitstream.Read(
+			transform_skip_rotation_enabled_flag,
+			transform_skip_context_enabled_flag,
+			implicit_rdpcm_enabled_flag,
+			explicit_rdpcm_enabled_flag,
+			extended_precision_processing_flag,
+			intra_smoothing_disabled_flag,
+			high_precision_offsets_enabled_flag,
+			persistent_rice_adaptation_enabled_flag,
+			cabac_bypass_alignment_enabled_flag);
+	}
+
+	void FNaluSPS::sps_scc_extension_t::Parse(uint32 const& in_chroma_format_idc, uint32 const& in_bit_depth_chroma_minus8, FBitstreamReader& Bitstream)
+	{
+		Bitstream.Read(
+			sps_curr_pic_ref_enabled_flag,
+			palette_mode_enabled_flag);
+
+		if (palette_mode_enabled_flag)
+		{
+			Bitstream.Read(
+				palette_max_size,
+				delta_palette_max_predictor_size,
+				sps_palette_predictor_initializers_present_flag);
+
+			if (sps_palette_predictor_initializers_present_flag)
+			{
+				Bitstream.Read(sps_num_palette_predictor_initializers_minus1);
+				const uint8 numComps = (in_chroma_format_idc == 0) ? 1u : 3u;
+				sps_palette_predictor_initializers.SetNumZeroed(numComps);
+				for (uint8 comp = 0; comp < numComps; comp++)
+				{
+					sps_palette_predictor_initializers[comp].SetNumUninitialized(sps_num_palette_predictor_initializers_minus1 + 1);
+					for (uint8 i = 0; i <= sps_num_palette_predictor_initializers_minus1; i++)
+					{
+						Bitstream.ReadBits(sps_palette_predictor_initializers[comp][i], in_bit_depth_chroma_minus8 + 8);
+					}
+				}
+			}
+		}
+
+		Bitstream.Read(
+			motion_vector_resolution_control_idc,
+			intra_boundary_filtering_disabled_flag);
+	}
+
+	FAVResult FNaluSPS::Parse()
+	{
+		// Get Bitstream
+		TArray64<uint8> RBSP;
+		RBSP.SetNumUninitialized(Size - StartCodeSize);
+
+		int32 RBSPsize = EBSPtoRBSP(RBSP.GetData(), EBSP + StartCodeSize, Size - StartCodeSize);
+		FBitstreamReader Bitstream(RBSP.GetData(), RBSPsize);
+
+		// Header which we have already parsed 
+		Bitstream.SkipBytes(2);
+
+		// Start parsing SPS
+		Bitstream.Read(
+			sps_video_parameter_set_id,
+			sps_max_sub_layers_minus1,
+			sps_temporal_id_nesting_flag);
+
+		profile_tier_level.Parse(1, sps_max_sub_layers_minus1, Bitstream);
+
+		Bitstream.Read(
+			sps_seq_parameter_set_id,
+			chroma_format_idc);
+
+		if (chroma_format_idc == 3)
+		{
+			Bitstream.Read(separate_colour_plane_flag);
+		}
+
+		ChromaArrayType = separate_colour_plane_flag == 0 ? (uint8)chroma_format_idc : 0u;
+
+		Bitstream.Read(
+			pic_width_in_luma_samples,
+			pic_height_in_luma_samples,
+			conformance_window_flag);
+
+		if (conformance_window_flag)
+		{
+			Bitstream.Read(
+				conf_win.left_offset,
+				conf_win.right_offset,
+				conf_win.top_offset,
+				conf_win.bottom_offset);
+		}
+
+		Bitstream.Read(
+			bit_depth_luma_minus8,
+			bit_depth_chroma_minus8,
+			log2_max_pic_order_cnt_lsb_minus4,
+			sps_sub_layer_ordering_info_present_flag);
+
+		sub_layer_ordering_infos.SetNumUninitialized(sps_max_sub_layers_minus1 + 1);
+
+		for (uint8 i = (sps_sub_layer_ordering_info_present_flag ? 0u : (uint8)sps_max_sub_layers_minus1); i <= sps_max_sub_layers_minus1; i++)
+		{
+			Bitstream.Read(
+				sub_layer_ordering_infos[i].sps_max_dec_pic_buffering_minus1,
+				sub_layer_ordering_infos[i].sps_max_num_reorder_pics,
+				sub_layer_ordering_infos[i].sps_max_latency_increase_plus1);
+		}
+
+		MaxDpbSize = sub_layer_ordering_infos.Last().sps_max_dec_pic_buffering_minus1 + 1;
+
+		Bitstream.Read(
+			log2_min_luma_coding_block_size_minus3,
+			log2_diff_max_min_luma_coding_block_size,
+			log2_min_luma_transform_block_size_minus2,
+			log2_diff_max_min_luma_transform_block_size,
+			max_transform_hierarchy_depth_inter,
+			max_transform_hierarchy_depth_intra,
+			scaling_list_enabled_flag);
+
+		if (scaling_list_enabled_flag)
+		{
+			Bitstream.Read(sps_scaling_list_data_present_flag);
+			if (sps_scaling_list_data_present_flag)
+			{
+				scaling_list_data.Parse(Bitstream);
+			}
+		}
+
+		Bitstream.Read(
+			amp_enabled_flag,
+			sample_adaptive_offset_enabled_flag,
+			pcm_enabled_flag);
+
+		if (pcm_enabled_flag)
+		{
+			Bitstream.Read(
+				pcm_sample_bit_depth_luma_minus1,
+				pcm_sample_bit_depth_chroma_minus1,
+				log2_min_pcm_luma_coding_block_size_minus3,
+				log2_diff_max_min_pcm_luma_coding_block_size,
+				pcm_loop_filter_disabled_flag);
+		}
+
+		Bitstream.Read(num_short_term_ref_pic_sets);
+		short_term_ref_pic_sets.SetNumZeroed(num_short_term_ref_pic_sets + 1); // We add one here as a slice can append a picture to this list so index "num_short_term_ref_pic_sets" is valid
+		for (uint8 i = 0; i < num_short_term_ref_pic_sets; i++)
+		{
+			short_term_ref_pic_sets[i].Parse(i, short_term_ref_pic_sets, Bitstream);
+			short_term_ref_pic_sets[i].CalculateValues(i, short_term_ref_pic_sets);
+		}
+
+		Bitstream.Read(long_term_ref_pics_present_flag);
+		if (long_term_ref_pics_present_flag)
+		{
+			Bitstream.Read(num_long_term_ref_pics_sps);
+			long_term_ref_pics_sps.SetNumUninitialized(num_long_term_ref_pics_sps);
+			for (uint8 i = 0; i < num_long_term_ref_pics_sps; i++)
+			{
+				long_term_ref_pics_sps[i].Parse(log2_max_pic_order_cnt_lsb_minus4, Bitstream);
+			}
+		}
+
+		Bitstream.Read(
+			sps_temporal_mvp_enabled_flag,
+			strong_intra_smoothing_enabled_flag,
+			vui_parameters_present_flag);
+
+		if (vui_parameters_present_flag)
+		{
+			vui_parameters.Parse(Bitstream);
+		}
+
+		Bitstream.Read(sps_extension_present_flag);
+
+		if (sps_extension_present_flag)
+		{
+			Bitstream.Read(
+				sps_range_extension_flag,
+				sps_multilayer_extension_flag,
+				sps_3d_extension_flag,
+				sps_scc_extension_flag,
+				sps_extension_4bits);
+		}
+
+		if (sps_range_extension_flag)
+		{
+			sps_range_extension.Parse(Bitstream);
+		}
+
+		if (sps_multilayer_extension_flag)
+		{
+			sps_multilayer_extension.Parse(Bitstream);
+		}
+
+		if (sps_3d_extension_flag)
+		{
+			sps_3d_extension.Parse(Bitstream);
+		}
+
+		if (sps_scc_extension_flag)
+		{
+			sps_scc_extension.Parse(chroma_format_idc, bit_depth_chroma_minus8, Bitstream);
+		}
+
+		if (sps_extension_4bits)
+		{
+			// TODO (aidan) read more extensions
+		}
+
+		return EAVResult::Success;
+	}
+
+	void FNaluPPS::pps_range_extension_t::Parse(uint8 const& in_transform_skip_enabled_flag, FBitstreamReader& Bitstream)
+	{
+		if (in_transform_skip_enabled_flag)
+		{
+			Bitstream.Read(log2_max_transform_skip_block_size_minus2);
+		}
+
+		Bitstream.Read(
+			cross_component_prediction_enabled_flag,
+			chroma_qp_offset_list_enabled_flag);
+
+		if (chroma_qp_offset_list_enabled_flag)
+		{
+			Bitstream.Read(
+				diff_cu_chroma_qp_offset_depth,
+				chroma_qp_offset_list_len_minus1);
+
+			cb_qp_offset_list.SetNumUninitialized(chroma_qp_offset_list_len_minus1 + 1);
+			cr_qp_offset_list.SetNumUninitialized(chroma_qp_offset_list_len_minus1 + 1);
+			for (uint8 i = 0; i <= chroma_qp_offset_list_len_minus1; i++)
+			{
+				Bitstream.Read(
+					cb_qp_offset_list[i],
+					cr_qp_offset_list[i]);
+			}
+		}
+
+		Bitstream.Read(
+			log2_sao_offset_scale_luma,
+			log2_sao_offset_scale_chroma);
+	}
+
+	void FNaluPPS::pps_scc_extension_t::Parse(FBitstreamReader& Bitstream)
+	{
+		Bitstream.Read(
+			pps_curr_pic_ref_enabled_flag,
+			residual_adaptive_colour_transform_enabled_flag);
+
+		if (residual_adaptive_colour_transform_enabled_flag)
+		{
+			Bitstream.Read(
+				pps_slice_act_qp_offsets_present_flag,
+				pps_act_y_qp_offset_plus5,
+				pps_act_cb_qp_offset_plus5,
+				pps_act_cr_qp_offset_plus3);
+		}
+
+		Bitstream.Read(pps_palette_predictor_initializers_present_flag);
+		if (pps_palette_predictor_initializers_present_flag)
+		{
+			Bitstream.Read(pps_num_palette_predictor_initializers);
+			if (pps_num_palette_predictor_initializers > 0)
+			{
+				Bitstream.Read(
+					monochrome_palette_flag,
+					luma_bit_depth_entry_minus8);
+
+				if (!monochrome_palette_flag)
+				{
+					Bitstream.Read(chroma_bit_depth_entry_minus8);
+				}
+
+				uint8 numComps = monochrome_palette_flag ? 1 : 3;
+				pps_palette_predictor_initializer.SetNumZeroed(numComps);
+				for (uint8 comp = 0; comp < numComps; comp++)
+				{
+					pps_palette_predictor_initializer[comp].SetNumZeroed(pps_num_palette_predictor_initializers);
+					for (uint8 i = 0; i < pps_num_palette_predictor_initializers; i++)
+					{
+						Bitstream.ReadBits(pps_palette_predictor_initializer[comp][i], (i >= 1) ? chroma_bit_depth_entry_minus8 + 8 : luma_bit_depth_entry_minus8 + 8);
+					}
+				}
+			}
+		}
+	}
+
+	FAVResult FNaluPPS::Parse()
+	{
+		// Get Bitstream
+		TArray64<uint8> RBSP;
+		RBSP.SetNumUninitialized(Size - StartCodeSize);
+
+		int32 RBSPsize = EBSPtoRBSP(RBSP.GetData(), EBSP + StartCodeSize, Size - StartCodeSize);
+		FBitstreamReader Bitstream(RBSP.GetData(), RBSPsize);
+
+		// Skip NAL Header which we have already parsed 
+		Bitstream.SkipBytes(2);
+
+		// Start parsing PPS
+		Bitstream.Read(
+			pps_pic_parameter_set_id,
+			pps_seq_parameter_set_id,
+			dependent_slice_segments_enabled_flag,
+			output_flag_present_flag,
+			num_extra_slice_header_bits,
+			sign_data_hiding_enabled_flag,
+			cabac_init_present_flag,
+			num_ref_idx_l0_default_active_minus1,
+			num_ref_idx_l1_default_active_minus1,
+			init_qp_minus26,
+			constrained_intra_pred_flag,
+			transform_skip_enabled_flag,
+			cu_qp_delta_enabled_flag);
+
+		if (cu_qp_delta_enabled_flag)
+		{
+			Bitstream.Read(diff_cu_qp_delta_depth);
+		}
+
+		Bitstream.Read(
+			pps_cb_qp_offset,
+			pps_cr_qp_offset,
+			pps_slice_chroma_qp_offsets_present_flag,
+			weighted_pred_flag,
+			weighted_bipred_flag,
+			transquant_bypass_enabled_flag,
+			tiles_enabled_flag,
+			entropy_coding_sync_enabled_flag);
+
+		if (tiles_enabled_flag)
+		{
+			Bitstream.Read(
+				num_tile_columns_minus1,
+				num_tile_rows_minus1,
+				uniform_spacing_flag);
+
+			if (!uniform_spacing_flag)
+			{
+				column_width_minus1.SetNumUninitialized(num_tile_columns_minus1 + 1);
+				for (uint8 i = 0; i < num_tile_columns_minus1; i++)
+				{
+					Bitstream.Read(column_width_minus1[i]);
+				}
+
+				row_height_minus1.SetNumUninitialized(num_tile_rows_minus1 + 1);
+				for (uint8 i = 0; i < num_tile_rows_minus1; i++)
+				{
+					Bitstream.Read(row_height_minus1[i]);
+				}
+
+				Bitstream.Read(loop_filter_across_tiles_enabled_flag);
+			}
+		}
+
+		Bitstream.Read(
+			pps_loop_filter_across_slices_enabled_flag,
+			deblocking_filter_control_present_flag);
+
+		if (deblocking_filter_control_present_flag)
+		{
+			Bitstream.Read(
+				deblocking_filter_override_enabled_flag,
+				pps_deblocking_filter_disabled_flag);
+
+			if (!pps_deblocking_filter_disabled_flag)
+			{
+				Bitstream.Read(
+					pps_beta_offset_div2,
+					pps_tc_offset_div2);
+			}
+		}
+
+		Bitstream.Read(pps_scaling_list_data_present_flag);
+		if (pps_scaling_list_data_present_flag)
+		{
+			scaling_list_data.Parse(Bitstream);
+		}
+
+		Bitstream.Read(
+			lists_modification_present_flag,
+			log2_parallel_merge_level_minus2,
+			slice_segment_header_extension_present_flag,
+			pps_extension_present_flag);
+
+		if (pps_extension_present_flag)
+		{
+			Bitstream.Read(
+				pps_range_extension_flag,
+				pps_multilayer_extension_flag,
+				pps_3d_extension_flag,
+				pps_scc_extension_flag,
+				pps_extension_4bits);
+		}
+
+		if (pps_range_extension_flag)
+		{
+			pps_range_extension.Parse(transform_skip_enabled_flag, Bitstream);
+		}
+
+		if (pps_multilayer_extension_flag)
+		{
+			pps_multilayer_extension.Parse(Bitstream);
+		}
+
+		if (pps_3d_extension_flag)
+		{
+			pps_3d_extension.Parse(Bitstream);
+		}
+
+		if (pps_scc_extension_flag)
+		{
+			pps_scc_extension.Parse(Bitstream);
+		}
+
+		if (pps_extension_4bits)
+		{
+			// TODO (aidan) read more extensions
+		}
+
+		return EAVResult::Success;
+	}
+
+	void FNaluSlice::ref_pic_list_modification_t::Parse(uint8 const& in_num_ref_idx_l0_active_minus1, uint8 const& in_num_ref_idx_l1_active_minus1, uint32 const& InNumPicTotalCurr, EH265SliceType const& SliceType, FBitstreamReader& BitStream)
+	{
+		BitStream.Read(ref_pic_list_modification_flag_l0);
+		list_entry_l0.SetNumZeroed(in_num_ref_idx_l0_active_minus1);
+		if(ref_pic_list_modification_flag_l0)
+		{
+			for(uint8 i = 0; i < in_num_ref_idx_l0_active_minus1; i++)
+			{
+				BitStream.ReadBits(list_entry_l0[i], FMath::CeilLogTwo(InNumPicTotalCurr));
+			}
+		}
+
+		if (SliceType == EH265SliceType::B)
+		{
+			BitStream.Read(ref_pic_list_modification_flag_l1);
+			list_entry_l1.SetNumZeroed(ref_pic_list_modification_flag_l1);
+			if(ref_pic_list_modification_flag_l1)
+			{
+				for(uint8 i = 0; i < in_num_ref_idx_l1_active_minus1; i++)
+				{
+					BitStream.ReadBits(list_entry_l1[i], FMath::CeilLogTwo(InNumPicTotalCurr));
+				}
+			}
+		}
+	}
+
+	void FNaluSlice::pred_weight_table_t::Parse(const uint8& ChromaArrayType, FNaluSlice const& CurrentSlice, FBitstreamReader& BitStream)
+	{
+		// TODO need to get acces to the nuh_layer_id on the slice to access this
+		/*
+		BitStream.Read(luma_log2_weight_denom);
+
+		if (ChromaArrayType != 0)
+		{
+			BitStream.Read(delta_chroma_log2_weight_denom);
+		}
+
+		for (uint8 i = 0; i <= CurrentSlice.num_ref_idx_l0_active_minus1; i++)
+		{
+			if ((CurrentSlice.RefPicList0[i].nuh_layer_id != CurrentSlice.nuh_layer_id) || (CurrentSlice.RefPicList0[i].PicOrderCntVal != CurrentSlice.PicOrderCntVal))
+			{
+				BitStream.Read(luma_weight_l0_flag[i]);
+			}
+		}
+
+		if (ChromaArrayType != 0)
+		{
+			for (uint8 i = 0; i <= CurrentSlice.num_ref_idx_l0_active_minus1; i++)
+			{
+				if ((CurrentSlice.RefPicList0[i].nuh_layer_id != CurrentSlice.nuh_layer_id) || (CurrentSlice.RefPicList0[i].PicOrderCntVal != CurrentSlice.PicOrderCntVal))
+				{
+					BitStream.Read(chroma_weight_l0_flag[i]);
+				}
+			}
+		}
+
+		for (uint8 i = 0; i <= CurrentSlice.num_ref_idx_l0_active_minus1; i++)
+		{
+			if (luma_weight_l0_flag[i])
+			{
+				BitStream.Read(delta_luma_weight_l0[i], luma_offset_l0[i]);
+			}
+			if (chroma_weight_l0_flag[i])
+			{
+				for (uint8 j = 0; j < 2; j++)
+				{
+					BitStream.Read(delta_chroma_weight_l0[i][j], delta_chroma_offset_l0[i][j]);
+				}
+			}
+		}
+
+		if (CurrentSlice.slice_type == EH265SliceType::B)
+		{
+			for (uint8 i = 0; i <= CurrentSlice.num_ref_idx_l1_active_minus1; i++)
+			{
+				if ((CurrentSlice.RefPicList1[i].nuh_layer_id != CurrentSlice.nuh_layer_id) || (CurrentSlice.RefPicList1[i].PicOrderCntVal != CurrentSlice.PicOrderCntVal))
+				{
+					BitStream.Read(luma_weight_l1_flag[i]);
+				}
+			}
+
+			if (ChromaArrayType != 0)
+			{
+				for (uint8 i = 0; i <= CurrentSlice.num_ref_idx_l1_active_minus1; i++)
+				{
+					if ((CurrentSlice.RefPicList1[i].nuh_layer_id != CurrentSlice.nuh_layer_id) || (CurrentSlice.RefPicList1[i].PicOrderCntVal != CurrentSlice.PicOrderCntVal))
+					{
+						BitStream.Read(chroma_weight_l1_flag[i]);
+					}
+				}
+			}
+
+			for (uint8 i = 0; i <= CurrentSlice.num_ref_idx_l1_active_minus1; i++)
+			{
+				if (luma_weight_l1_flag[i])
+				{
+					BitStream.Read(delta_luma_weight_l1[i], luma_offset_l1[i]);
+				}
+				if (chroma_weight_l1_flag[i])
+				{
+					for (uint8 j = 0; j < 2; j++)
+					{
+						BitStream.Read(delta_chroma_weight_l1[i][j], delta_chroma_offset_l1[i][j]);
+					}
+				}
+			}
+		}
+		*/
+	}
+	
+	FAVResult FNaluSlice::Parse(FVideoDecoderConfigH265* InConfig)
+	{
+		// Get Bitstream
+		TArray64<uint8> RBSP;
+		RBSP.SetNumUninitialized(Size);
+
+		int32 RBSPsize = EBSPtoRBSP(RBSP.GetData(), EBSP, Size);
+		FBitstreamReader Bitstream(RBSP.GetData(), RBSPsize);
+
+		// Skip NALU delimiter and NAL Header which we have already parsed 
+		Bitstream.SkipBytes(StartCodeSize + 2);
+
+		// Start parsing SliceHeader
+		Bitstream.Read(first_slice_segment_in_pic_flag);
+
+		if(IsIDR())
+		{
+			CurrPicIdx = InConfig->CurrPicIdx = 0;
+		}
+		else
+		{
+			CurrPicIdx = first_slice_segment_in_pic_flag ? InConfig->CurrPicIdx++ : InConfig->CurrPicIdx;
+		}
+		
+		if (nal_unit_type >= ENaluType::BLA_W_LP && nal_unit_type <= ENaluType::RSV_IRAP_VCL23)
+		{
+			Bitstream.Read(no_output_of_prior_pics_flag);
+		}
+
+		Bitstream.Read(slice_pic_parameter_set_id);
+
+		if (InConfig->ParsedPPS.IsEmpty() || InConfig->ParsedSPS.IsEmpty() || InConfig->ParsedVPS.IsEmpty())
+		{
+			return EAVResult::PendingInput;
+		}
+		
+		CurrentPPS = InConfig->ParsedPPS.FindRef(slice_pic_parameter_set_id);
+		auto PinnedPPS = CurrentPPS.Pin();
+
+		CurrentSPS = InConfig->ParsedSPS.FindRef(PinnedPPS->pps_seq_parameter_set_id);
+		auto PinnedSPS = CurrentSPS.Pin();
+
+		CurrentVPS = InConfig->ParsedVPS.FindRef(PinnedSPS->sps_video_parameter_set_id);
+		auto PinnedVPS = CurrentVPS.Pin();
+		
+		if (!PinnedPPS || !PinnedSPS || !PinnedVPS)
+		{
+			return EAVResult::PendingInput;
+		}
+
+		if (!first_slice_segment_in_pic_flag)
+		{
+			if (dependent_slice_segment_flag)
+			{
+				Bitstream.Read(dependent_slice_segment_flag);
+			}
+
+			// TODO (aidan) Should I pre calculate this on the SPS
+			const uint32 MinCbLog2SizeY = PinnedSPS->log2_min_luma_coding_block_size_minus3 + 3;
+			const uint32 CtbLog2SizeY = MinCbLog2SizeY + PinnedSPS->log2_diff_max_min_luma_coding_block_size;
+			const uint32 CtbSizeY = 1u << CtbLog2SizeY;
+			const uint32 PicWidthInCtbsY = FMath::DivideAndRoundUp<uint32>(PinnedSPS->pic_width_in_luma_samples, CtbSizeY);
+			const uint32 PicHeightInCtbsY = FMath::DivideAndRoundUp<uint32>(PinnedSPS->pic_height_in_luma_samples, CtbSizeY);
+			const uint32 PicSizeInCtbsY = PicWidthInCtbsY * PicHeightInCtbsY;
+
+			Bitstream.ReadBits(slice_segment_address, FMath::CeilLogTwo(PicSizeInCtbsY));
+		}
+
+		uint32 CuQpDeltaVal = 0;
+		if (!dependent_slice_segment_flag)
+		{
+			Bitstream.SkipBits(PinnedPPS->num_extra_slice_header_bits);
+			Bitstream.Read(slice_type);
+
+			if (PinnedPPS->output_flag_present_flag)
+			{
+				Bitstream.Read(pic_output_flag);
+			}
+
+			if (PinnedSPS->separate_colour_plane_flag == 1)
+			{
+				Bitstream.Read(colour_plane_id);
+			}
+
+			if (nal_unit_type != ENaluType::IDR_W_RADL && nal_unit_type != ENaluType::IDR_N_LP)
+			{
+				Bitstream.ReadBits(slice_pic_order_cnt_lsb, PinnedSPS->log2_max_pic_order_cnt_lsb_minus4 + 4);
+				Bitstream.Read(short_term_ref_pic_set_sps_flag);
+
+				if (!short_term_ref_pic_set_sps_flag)
+				{
+					NumBitsForShortTermRPSInSlice = Bitstream.NumBitsRemaining();
+					PinnedSPS->short_term_ref_pic_sets[PinnedSPS->num_short_term_ref_pic_sets].Parse(PinnedSPS->num_short_term_ref_pic_sets, PinnedSPS->short_term_ref_pic_sets, Bitstream);
+					PinnedSPS->short_term_ref_pic_sets[PinnedSPS->num_short_term_ref_pic_sets].CalculateValues(PinnedSPS->num_short_term_ref_pic_sets, PinnedSPS->short_term_ref_pic_sets);
+					NumBitsForShortTermRPSInSlice -= Bitstream.NumBitsRemaining();
+				}
+				else if (PinnedSPS->num_short_term_ref_pic_sets > 1)
+				{
+					Bitstream.ReadBits(short_term_ref_pic_set_idx, FMath::CeilLogTwo(PinnedSPS->num_short_term_ref_pic_sets));
+				}
+
+				if (PinnedSPS->long_term_ref_pics_present_flag)
+				{
+					if (PinnedSPS->num_long_term_ref_pics_sps > 0)
+					{
+						Bitstream.Read(num_long_term_sps);
+					}
+
+					Bitstream.Read(num_long_term_pics);
+
+					lt_idx_sps.SetNumZeroed(num_long_term_sps);
+					poc_lsb_lt.SetNumZeroed(num_long_term_sps + num_long_term_pics);
+					used_by_curr_pic_lt_flag.SetNumZeroed(num_long_term_sps + num_long_term_pics);
+					delta_poc_msb_present_flag.SetNumZeroed(num_long_term_sps + num_long_term_pics);
+					delta_poc_msb_cycle_lt.SetNumZeroed(num_long_term_sps + num_long_term_pics);
+
+					for (uint16 i = 0; i < num_long_term_sps + num_long_term_pics; i++)
+					{
+						if (i < num_long_term_sps)
+						{
+							if (num_long_term_sps > 1)
+							{
+								Bitstream.ReadBits(lt_idx_sps[i], FMath::CeilLogTwo(PinnedSPS->num_long_term_ref_pics_sps));
+							}
+						}
+						else
+						{
+							Bitstream.ReadBits(poc_lsb_lt[i], PinnedSPS->log2_max_pic_order_cnt_lsb_minus4 + 4);
+							Bitstream.Read(used_by_curr_pic_lt_flag[i]);
+						}
+
+						Bitstream.Read(delta_poc_msb_present_flag[i]);
+						if (delta_poc_msb_present_flag[i])
+						{
+							Bitstream.Read(delta_poc_msb_cycle_lt[i]);
+						}
+					}
+				}
+
+				if (PinnedSPS->sample_adaptive_offset_enabled_flag)
+				{
+					Bitstream.Read(slice_sao_luma_flag);
+					if (PinnedSPS->ChromaArrayType != 0)
+					{
+						Bitstream.Read(slice_sao_chroma_flag);
+					}
+				}
+
+				// Derive CurrRpsIdx (7.4.7.1)
+				CurrRpsIdx = PinnedSPS->num_short_term_ref_pic_sets;
+
+				if (short_term_ref_pic_set_sps_flag)
+				{
+					CurrRpsIdx = short_term_ref_pic_set_idx;
+				}
+
+				// Derive NumPicTotalCurr (7.4.7.2)
+				uint32 NumPicTotalCurr = 0;
+
+				for (uint32 i = 0; i < PinnedSPS->short_term_ref_pic_sets[CurrRpsIdx].NumNegativePics; i++)
+				{
+					if (PinnedSPS->short_term_ref_pic_sets[CurrRpsIdx].UsedByCurrPicS0[i])
+					{
+						NumPicTotalCurr++;
+					}
+				}
+
+				for (uint32 i = 0; i < PinnedSPS->short_term_ref_pic_sets[CurrRpsIdx].num_positive_pics; i++)
+				{
+					if (PinnedSPS->short_term_ref_pic_sets[CurrRpsIdx].UsedByCurrPicS1[i])
+					{
+						NumPicTotalCurr++;
+					}
+				}
+
+				TArray<uint32> UsedByCurrPicLt;
+				UsedByCurrPicLt.SetNumZeroed(num_long_term_sps + num_long_term_pics);
+				for (uint32 i = 0; i < num_long_term_sps + num_long_term_pics; i++)
+				{
+					if (i < num_long_term_sps)
+					{
+						UsedByCurrPicLt[i] = PinnedSPS->long_term_ref_pics_sps[lt_idx_sps[i]].used_by_curr_pic_lt_sps_flag;
+					}
+					else
+					{
+						UsedByCurrPicLt[i] = used_by_curr_pic_lt_flag[i];
+					}
+
+					if (UsedByCurrPicLt[i])
+					{
+						NumPicTotalCurr++;
+					}
+				}
+
+				if (PinnedPPS->pps_scc_extension.pps_curr_pic_ref_enabled_flag)
+				{
+					NumPicTotalCurr++;
+				}
+
+				// Derive PicOrderCntMsb (8.3.1)
+				uint32 MaxPicOrderCntLsb = 1u << (PinnedSPS->log2_max_pic_order_cnt_lsb_minus4 + 4);
+				uint8 PicOrderCntMsb;
+				if ((slice_pic_order_cnt_lsb < PinnedVPS->prevPicOrderCntLsb) && ((PinnedVPS->prevPicOrderCntLsb - slice_pic_order_cnt_lsb) >= (MaxPicOrderCntLsb / 2)))
+				{
+					PicOrderCntMsb = PinnedVPS->prevPicOrderCntMsb + MaxPicOrderCntLsb;
+				}
+				else if ((slice_pic_order_cnt_lsb > PinnedVPS->prevPicOrderCntLsb) && ((slice_pic_order_cnt_lsb - PinnedVPS->prevPicOrderCntLsb) > (MaxPicOrderCntLsb / 2)))
+				{
+					PicOrderCntMsb = PinnedVPS->prevPicOrderCntMsb - MaxPicOrderCntLsb;
+				}
+				else
+				{
+					PicOrderCntMsb = PinnedVPS->prevPicOrderCntMsb;
+				}
+
+				CurrPicOrderCntVal = PicOrderCntMsb + slice_pic_order_cnt_lsb;
+
+				if (slice_type == EH265SliceType::P || slice_type == EH265SliceType::B)
+				{
+					Bitstream.Read(num_ref_idx_active_override_flag);
+					if (num_ref_idx_active_override_flag)
+					{
+						Bitstream.Read(num_ref_idx_l0_active_minus1);
+						if (slice_type == EH265SliceType::B)
+						{
+							Bitstream.Read(num_ref_idx_l1_active_minus1);
+						}
+					}
+					else
+					{
+						num_ref_idx_l0_active_minus1 = PinnedPPS->num_ref_idx_l0_default_active_minus1;
+						if (slice_type == EH265SliceType::B)
+						{
+							num_ref_idx_l1_active_minus1 = PinnedPPS->num_ref_idx_l1_default_active_minus1;
+						}
+					}
+
+					if (PinnedPPS->lists_modification_present_flag && NumPicTotalCurr > 1)
+					{
+						ref_pic_list_modification.Parse(num_ref_idx_l0_active_minus1, num_ref_idx_l1_active_minus1, NumPicTotalCurr, slice_type, Bitstream);
+					}
+
+					if (slice_type == EH265SliceType::B)
+					{
+						Bitstream.Read(mvd_l1_zero_flag);
+					}
+
+					if (PinnedPPS->cabac_init_present_flag)
+					{
+						Bitstream.Read(cabac_init_flag);
+					}
+
+					if (slice_temporal_mvp_enabled_flag)
+					{
+						if (slice_type == EH265SliceType::B)
+						{
+							Bitstream.Read(collocated_from_l0_flag);
+						}
+
+						if ((collocated_from_l0_flag && num_ref_idx_l0_active_minus1 > 0) || (!collocated_from_l0_flag && num_ref_idx_l1_active_minus1 > 0))
+						{
+							Bitstream.Read(collocated_ref_idx);
+						}
+					}
+
+					if ((PinnedPPS->weighted_pred_flag && slice_type == EH265SliceType::P) || (PinnedPPS->weighted_bipred_flag && slice_type == EH265SliceType::B))
+					{
+						pred_weight_table.Parse(PinnedSPS->ChromaArrayType, *this, Bitstream);
+					}
+
+					Bitstream.Read(five_minus_max_num_merge_cand);
+
+					if (PinnedSPS->sps_scc_extension.motion_vector_resolution_control_idc == 2)
+					{
+						Bitstream.Read(use_integer_mv_flag);	
+					}
+				}
+
+				Bitstream.Read(slice_qp_delta);
+
+				if (PinnedPPS->pps_slice_chroma_qp_offsets_present_flag)
+				{
+					Bitstream.Read(
+						slice_cb_qp_offset,
+						slice_cr_qp_offset);
+				}
+
+				if (PinnedPPS->pps_range_extension.chroma_qp_offset_list_enabled_flag)
+				{
+					Bitstream.Read(cu_chroma_qp_offset_enabled_flag);
+				}
+
+				if (PinnedPPS->deblocking_filter_override_enabled_flag)
+				{
+					Bitstream.Read(deblocking_filter_override_flag);
+				}
+
+				if (deblocking_filter_override_flag)
+				{
+					Bitstream.Read(slice_deblocking_filter_disabled_flag);
+					if (!slice_deblocking_filter_disabled_flag)
+					{
+						Bitstream.Read(
+							slice_beta_offset_div2,
+							slice_tc_offset_div2);
+					}
+				}
+
+				if (PinnedPPS->pps_loop_filter_across_slices_enabled_flag && (slice_sao_luma_flag || slice_sao_chroma_flag || !slice_deblocking_filter_disabled_flag))
+				{
+					Bitstream.Read(slice_loop_filter_across_slices_enabled_flag);
+				}
+			}
+
+			if (PinnedPPS->tiles_enabled_flag || PinnedPPS->entropy_coding_sync_enabled_flag)
+			{
+				Bitstream.Read(num_entry_point_offsets);
+				if (num_entry_point_offsets > 0)
+				{
+					Bitstream.Read(offset_len_minus1);
+					entry_point_offset_minus1.SetNumUninitialized(num_entry_point_offsets);
+
+					for (uint8 i = 0; i < num_entry_point_offsets; i++)
+					{
+						Bitstream.ReadBits(entry_point_offset_minus1[i], offset_len_minus1 + 1);
+					}
+				}
+			}
+
+			if (PinnedPPS->slice_segment_header_extension_present_flag)
+			{
+				// TODO (aidan) handle slice extensions
+			}
+		}
+
+		return EAVResult::Success;
+	}
+
+	FAVResult FNaluSEI::Parse()
+	{
+		// Get Bitstream
+		TArray64<uint8> RBSP;
+		RBSP.SetNumUninitialized(Size - StartCodeSize);
+
+		int32 RBSPsize = EBSPtoRBSP(RBSP.GetData(), EBSP + StartCodeSize, Size - StartCodeSize);
+		FBitstreamReader Bitstream(RBSP.GetData(), RBSPsize);
+
+		// Skip NAL Header which we have already parsed 
+		Bitstream.SkipBytes(2);
+
+		// Start parsing SEI
+		return FAVResult(EAVResult::Success, TEXT("SEI Unimplemented"));
+	}
+
+} // namespace UE::AVCodecCore::H265
