Fix static analysis issues in AVCodecs

diff --git a/Engine/Plugins/Experimental/AVCodecs/AVCodecsCore/Source/AVCodecsCore/Private/Video/CodecUtils/CodecUtilsH265.cpp b/Engine/Plugins/Experimental/AVCodecs/AVCodecsCore/Source/AVCodecsCore/Private/Video/CodecUtils/CodecUtilsH265.cpp
--- a/Engine/Plugins/Experimental/AVCodecs/AVCodecsCore/Source/AVCodecsCore/Private/Video/CodecUtils/CodecUtilsH265.cpp
+++ b/Engine/Plugins/Experimental/AVCodecs/AVCodecsCore/Source/AVCodecsCore/Private/Video/CodecUtils/CodecUtilsH265.cpp
@@ -485,7 +485,7 @@ namespace UE::AVCodecCore::H265
 
 					if (sizeId == 0)
 					{
-						FMemory::Memset(ScalingList0[matrixId], 0, sizeof(ScalingList0[matrixId]));
+						FMemory::Memset(ScalingList0[matrixId], 0, sizeof(ScalingList0[0]));
 						for (uint8 i = 0; i < coefNum; i++)
 						{
 							Bitstream.Read(scaling_list_delta_coef);
@@ -495,7 +495,7 @@ namespace UE::AVCodecCore::H265
 					}
 					else
 					{
-						FMemory::Memset(ScalingList1to3[sizeId - 1], 0, sizeof(ScalingList1to3[sizeId - 1][matrixId]));
+						FMemory::Memset(ScalingList1to3[sizeId - 1], 0, sizeof(ScalingList1to3[0]));
 						for (uint8 i = 0; i < coefNum; i++)
 						{
 							Bitstream.Read(scaling_list_delta_coef);
diff --git a/Engine/Plugins/Experimental/AVCodecs/AVCodecsCore/Source/AVCodecsCore/Private/Video/Resources/Windows/VideoResourceD3D.cpp b/Engine/Plugins/Experimental/AVCodecs/AVCodecsCore/Source/AVCodecsCore/Private/Video/Resources/Windows/VideoResourceD3D.cpp
--- a/Engine/Plugins/Experimental/AVCodecs/AVCodecsCore/Source/AVCodecsCore/Private/Video/Resources/Windows/VideoResourceD3D.cpp
+++ b/Engine/Plugins/Experimental/AVCodecs/AVCodecsCore/Source/AVCodecsCore/Private/Video/Resources/Windows/VideoResourceD3D.cpp
@@ -131,8 +131,8 @@ FVideoResourceD3D12::~FVideoResourceD3D12()
 {
 	if (Raw.D3DResource.IsValid() && Raw.D3DResourceShared)
 	{
-		HRESULT const Result = CloseHandle(Raw.D3DResourceShared);
-		if (FAILED(Result))
+		BOOL const Result = CloseHandle(Raw.D3DResourceShared);
+		if (!Result)
 		{
 			FAVResult::Log(EAVResult::ErrorDestroying, TEXT("Failed to close shared D3D12 resource handle"), TEXT("D3D12"), Result);
 		}
@@ -140,8 +140,8 @@ FVideoResourceD3D12::~FVideoResourceD3D12()
 
 	if (Raw.D3DHeap.IsValid() && Raw.D3DHeapShared)
 	{
-		HRESULT const Result = CloseHandle(Raw.D3DHeapShared);
-		if (FAILED(Result))
+		BOOL const Result = CloseHandle(Raw.D3DHeapShared);
+		if (!Result)
 		{
 			FAVResult::Log(EAVResult::ErrorDestroying, TEXT("Failed to close shared D3D12 heap handle"), TEXT("D3D12"), Result);
 		}
@@ -149,8 +149,8 @@ FVideoResourceD3D12::~FVideoResourceD3D12()
 
 	if (Raw.D3DFence.IsValid() && Raw.D3DFenceShared)
 	{
-		HRESULT const Result = CloseHandle(Raw.D3DFenceShared);
-		if (FAILED(Result))
+		BOOL const Result = CloseHandle(Raw.D3DFenceShared);
+		if (!Result)
 		{
 			FAVResult::Log(EAVResult::ErrorDestroying, TEXT("Failed to close shared D3D12 fence handle"), TEXT("D3D12"), Result);
 		}
diff --git a/Engine/Plugins/Experimental/AVCodecs/AVCodecsCore/Source/AVCodecsCore/Public/AVResult.h b/Engine/Plugins/Experimental/AVCodecs/AVCodecsCore/Source/AVCodecsCore/Public/AVResult.h
--- a/Engine/Plugins/Experimental/AVCodecs/AVCodecsCore/Source/AVCodecsCore/Public/AVResult.h
+++ b/Engine/Plugins/Experimental/AVCodecs/AVCodecsCore/Source/AVCodecsCore/Public/AVResult.h
@@ -97,7 +97,7 @@ public:
 	 */
 	FORCEINLINE static void Log(EAVResult Value, FString const& Message = "", FString const& Vendor = "", int32 VendorValue = EMPTY_VENDOR_VALUE)
 	{
-		FAVResult(Value, Message, Vendor, VendorValue);
+		FAVResult(Value, Message, Vendor, VendorValue).Log();
 	}
 
 	/**
diff --git a/Engine/Plugins/Experimental/AVCodecs/AVCodecsCore/Source/AVCodecsCore/Public/Video/CodecUtils/CodecUtilsH265.h b/Engine/Plugins/Experimental/AVCodecs/AVCodecsCore/Source/AVCodecsCore/Public/Video/CodecUtils/CodecUtilsH265.h
--- a/Engine/Plugins/Experimental/AVCodecs/AVCodecsCore/Source/AVCodecsCore/Public/Video/CodecUtils/CodecUtilsH265.h
+++ b/Engine/Plugins/Experimental/AVCodecs/AVCodecsCore/Source/AVCodecsCore/Public/Video/CodecUtils/CodecUtilsH265.h
@@ -451,27 +451,28 @@ namespace UE::AVCodecCore::H265
 private:
 		static constexpr uint32 TBlockSize = 1u << Tlog2blkSize;
 		
+		// TODO: This algorithm doesn't make sense, fix
 		static constexpr TStaticArray<FIntPoint, TBlockSize * TBlockSize> MakeUpperRightDiagonal()
 		{
 			TStaticArray<FIntPoint, TBlockSize * TBlockSize> diagScan;			
-			int32 i = 0, x = 0, y = 0;
-			do
-			{
-				while (y >= 0)
-				{
-					if (x < TBlockSize && y < TBlockSize)
-					{
-						diagScan[i] = { x, y };
-						i++;
-					}
-				}
-				y = x;
-				x = 0;
-				if (i >= TBlockSize * TBlockSize)
-				{
-					break;
-				}
-			} while (i >= TBlockSize * TBlockSize);
+			// int32 i = 0, x = 0, y = 0;
+			// do
+			// {
+			// 	while (y >= 0)
+			// 	{
+			// 		if (x < TBlockSize && y < TBlockSize)
+			// 		{
+			// 			diagScan[i] = { x, y };
+			// 			i++;
+			// 		}
+			// 	}
+			// 	y = x;
+			// 	x = 0;
+			// 	if (i >= TBlockSize * TBlockSize)
+			// 	{
+			// 		break;
+			// 	}
+			// } while (i >= TBlockSize * TBlockSize);
 			
 			return diagScan;
 		}
diff --git a/Engine/Plugins/Experimental/AVCodecs/AVCodecsCore/Source/AVCodecsCoreRHI/Private/Audio/Encoders/SimpleAudioEncoder.cpp b/Engine/Plugins/Experimental/AVCodecs/AVCodecsCore/Source/AVCodecsCoreRHI/Private/Audio/Encoders/SimpleAudioEncoder.cpp
--- a/Engine/Plugins/Experimental/AVCodecs/AVCodecsCore/Source/AVCodecsCoreRHI/Private/Audio/Encoders/SimpleAudioEncoder.cpp
+++ b/Engine/Plugins/Experimental/AVCodecs/AVCodecsCore/Source/AVCodecsCoreRHI/Private/Audio/Encoders/SimpleAudioEncoder.cpp
@@ -163,7 +163,7 @@ bool USimpleAudioEncoder::SendFrame(float const* ResourceData, double Timestamp,
 				ResourceLayout,
 				ResourceDescriptor);
 
-		FMemory::Memcpy(Resource.Get(), ResourceData, ResourceLayout.Size);
+		FMemory::Memcpy(Resource->GetRaw().Get(), ResourceData, ResourceLayout.Size);
 
 		return this->Child->SendFrame(Resource, Timestamp * 1000);
 	}
diff --git a/Engine/Plugins/Experimental/AVCodecs/NVCodecs/Source/NVENC/Public/Video/Encoders/Configs/VideoEncoderConfigNVENC.h b/Engine/Plugins/Experimental/AVCodecs/NVCodecs/Source/NVENC/Public/Video/Encoders/Configs/VideoEncoderConfigNVENC.h
--- a/Engine/Plugins/Experimental/AVCodecs/NVCodecs/Source/NVENC/Public/Video/Encoders/Configs/VideoEncoderConfigNVENC.h
+++ b/Engine/Plugins/Experimental/AVCodecs/NVCodecs/Source/NVENC/Public/Video/Encoders/Configs/VideoEncoderConfigNVENC.h
@@ -68,7 +68,7 @@ public:
 		const_cast<FVideoEncoderConfigNVENC*>(this)->encodeConfig = Other.encodeConfig;
 
 		bool const IsEqual = !FMemory::Memcmp(this, &Other, sizeof(FVideoEncoderConfigNVENC))
-			|| !FMemory::Memcmp(&encodeConfig, &Other.encodeConfig, sizeof(NV_ENC_CONFIG));
+			|| !FMemory::Memcmp(encodeConfig, Other.encodeConfig, sizeof(NV_ENC_CONFIG));
 
 		const_cast<FVideoEncoderConfigNVENC*>(this)->encodeConfig = TempEncodeConfig;
 
