Fix AVCodecs NonUnity compile errors

diff --git a/Engine/Plugins/Experimental/AVCodecs/AMFCodecs/Source/AMFCodecsRHI/Private/AMFCodecsRHIModule.cpp b/Engine/Plugins/Experimental/AVCodecs/AMFCodecs/Source/AMFCodecsRHI/Private/AMFCodecsRHIModule.cpp
--- a/Engine/Plugins/Experimental/AVCodecs/AMFCodecs/Source/AMFCodecsRHI/Private/AMFCodecsRHIModule.cpp
+++ b/Engine/Plugins/Experimental/AVCodecs/AMFCodecs/Source/AMFCodecsRHI/Private/AMFCodecsRHIModule.cpp
@@ -3,6 +3,7 @@
 #include "Modules/ModuleManager.h"
 
 #include "IVulkanDynamicRHI.h"
+#include "Misc/App.h"
 
 #if PLATFORM_WINDOWS
 #include "ID3D11DynamicRHI.h"
diff --git a/Engine/Plugins/Experimental/AVCodecs/AVCodecsCore/Source/AVCodecsCoreRHI/Private/AVCodecsCoreRHIModule.cpp b/Engine/Plugins/Experimental/AVCodecs/AVCodecsCore/Source/AVCodecsCoreRHI/Private/AVCodecsCoreRHIModule.cpp
--- a/Engine/Plugins/Experimental/AVCodecs/AVCodecsCore/Source/AVCodecsCoreRHI/Private/AVCodecsCoreRHIModule.cpp
+++ b/Engine/Plugins/Experimental/AVCodecs/AVCodecsCore/Source/AVCodecsCoreRHI/Private/AVCodecsCoreRHIModule.cpp
@@ -3,6 +3,7 @@
 #include "Modules/ModuleManager.h"
 
 #include "IVulkanDynamicRHI.h"
+#include "Misc/App.h"
 
 #include "Video/Resources/VideoResourceRHI.h"
 #include "Video/Resources/VideoResourceVulkan.h"
