Add: Allow users to specify transform coding and entropy coding mode in H264 config

diff --git a/Engine/Plugins/Experimental/AVCodecs/AVCodecsCore/Source/AVCodecsCore/Public/Video/CodecUtils/CodecUtilsH264.h b/Engine/Plugins/Experimental/AVCodecs/AVCodecsCore/Source/AVCodecsCore/Public/Video/CodecUtils/CodecUtilsH264.h
--- a/Engine/Plugins/Experimental/AVCodecs/AVCodecsCore/Source/AVCodecsCore/Public/Video/CodecUtils/CodecUtilsH264.h
+++ b/Engine/Plugins/Experimental/AVCodecs/AVCodecsCore/Source/AVCodecsCore/Public/Video/CodecUtils/CodecUtilsH264.h
@@ -41,6 +41,20 @@ enum class EH264Profile : uint8
 	MAX
 };
 
+enum class EH264AdaptiveTransformMode : uint8
+{
+	Auto,
+	Disable,
+	Enable
+};
+
+enum class EH264EntropyCodingMode : uint8
+{
+	Auto,
+	CABAC,
+	CAVLC
+};
+
 extern AVCODECSCORE_API FH264ProfileDefinition GH264ProfileDefinitions[static_cast<uint8>(EH264Profile::MAX)];
 
 namespace UE::AVCodecCore::H264
@@ -80,7 +94,7 @@ namespace UE::AVCodecCore::H264
 		Unspecified30 = 30,
 		Unspecified31 = 31
 	};
-	
+
 	enum class EH264ProfileIDC : uint8
 	{
 		Auto = 0,
@@ -196,13 +210,13 @@ namespace UE::AVCodecCore::H264
 		Component = 0,
 		PAL = 1,
 		NTSC = 2,
-		SECAM =3,
+		SECAM = 3,
 		MAC = 4,
 		Unspecified = 5,
 		Reserved6 = 6,
 		Reserved7 = 7
 	};
-}
+} // namespace UE::AVCodecCore::H264
 
 struct FH264ProfileDefinition
 {
@@ -217,7 +231,7 @@ namespace UE::AVCodecCore::H264
 	inline int32 EBSPtoRBSP(uint8* OutBuf, const uint8* InBuf, int32 NumBytesIn)
 	{
 		uint8* OutBase = OutBuf;
-		while(NumBytesIn-- > 0)
+		while (NumBytesIn-- > 0)
 		{
 			uint8 b = *InBuf++;
 			*OutBuf++ = b;
@@ -236,7 +250,7 @@ namespace UE::AVCodecCore::H264
 		}
 		return OutBuf - OutBase;
 	}
-	
+
 	struct FNaluInfo
 	{
 		uint64 Start, Size;
@@ -245,78 +259,79 @@ namespace UE::AVCodecCore::H264
 		ENaluType Type;
 		const uint8* Data;
 	};
-	
+
 	FAVResult FindNALUs(FVideoPacket const& InPacket, TArray<FNaluInfo>& FoundNalus);
 
 	constexpr uint8 Default_4x4_Intra[16] = {
-		6,13,13,20,
-		20,20,28,28,
-		28,28,32,32,
-		32,37,37,42
+		6, 13, 13, 20,
+		20, 20, 28, 28,
+		28, 28, 32, 32,
+		32, 37, 37, 42
 	};
-	
+
 	constexpr uint8 Default_4x4_Inter[16] = {
-		10,14,14,20,
-		20,20,24,24,
-		24,24,27,27,
-		27,30,30,34
+		10, 14, 14, 20,
+		20, 20, 24, 24,
+		24, 24, 27, 27,
+		27, 30, 30, 34
 	};
 
 	constexpr uint8 Default_8x8_Intra[64] = {
-		6,10,10,13,11,13,16,16,
-		16,16,18,18,18,18,18,23,
-		23,23,23,23,23,25,25,25,
-		25,25,25,25,27,27,27,27,
-		27,27,27,27,29,29,29,29,
-		29,29,29,31,31,31,31,31,
-		31,33,33,33,33,33,36,36,
-		36,36,38,38,38,40,40,42 };
-	
+		6, 10, 10, 13, 11, 13, 16, 16,
+		16, 16, 18, 18, 18, 18, 18, 23,
+		23, 23, 23, 23, 23, 25, 25, 25,
+		25, 25, 25, 25, 27, 27, 27, 27,
+		27, 27, 27, 27, 29, 29, 29, 29,
+		29, 29, 29, 31, 31, 31, 31, 31,
+		31, 33, 33, 33, 33, 33, 36, 36,
+		36, 36, 38, 38, 38, 40, 40, 42
+	};
+
 	constexpr uint8 Default_8x8_Inter[64] = {
-		9,13,13,15,13,15,17,17,
-		17,17,19,19,19,19,19,21,
-		21,21,21,21,21,22,22,22,
-		22,22,22,22,24,24,24,24,
-		24,24,24,24,25,25,25,25,
-		25,25,25,27,27,27,27,27,
-		27,28,28,28,28,28,30,30,
-		30,30,32,32,32,33,33,35
+		9, 13, 13, 15, 13, 15, 17, 17,
+		17, 17, 19, 19, 19, 19, 19, 21,
+		21, 21, 21, 21, 21, 22, 22, 22,
+		22, 22, 22, 22, 24, 24, 24, 24,
+		24, 24, 24, 24, 25, 25, 25, 25,
+		25, 25, 25, 27, 27, 27, 27, 27,
+		27, 28, 28, 28, 28, 28, 30, 30,
+		30, 30, 32, 32, 32, 33, 33, 35
 	};
 
 	constexpr uint8 ZigZag4x4[16] = {
-		0,1,5,6,
-		2,4,7,12,
-		3,8,11,13,
-		9,10,15,15
+		0, 1, 5, 6,
+		2, 4, 7, 12,
+		3, 8, 11, 13,
+		9, 10, 15, 15
 	};
 
 	constexpr uint8 FieldScan4x4[16] = {
-		0,2,8,12,
-		1,5,9,13,
-		3,6,10,14,
-		4,7,11,15
+		0, 2, 8, 12,
+		1, 5, 9, 13,
+		3, 6, 10, 14,
+		4, 7, 11, 15
 	};
 
 	constexpr uint8 ZigZag8x8[64] = {
-		0,1,5,6,14,15,27,28,
-		2,4,7,13,16,26,29,42,
-		3,8,12,17,25,30,41,43,
-		9,11,18,24,31,40,44,53,
-		10,19,23,32,39,45,52,54,
-		20,22,33,38,46,51,55,60,
-		21,34,37,47,50,56,59,61,
-		35,36,48,49,57,58,62,63
+		0, 1, 5, 6, 14, 15, 27, 28,
+		2, 4, 7, 13, 16, 26, 29, 42,
+		3, 8, 12, 17, 25, 30, 41, 43,
+		9, 11, 18, 24, 31, 40, 44, 53,
+		10, 19, 23, 32, 39, 45, 52, 54,
+		20, 22, 33, 38, 46, 51, 55, 60,
+		21, 34, 37, 47, 50, 56, 59, 61,
+		35, 36, 48, 49, 57, 58, 62, 63
 	};
 
 	constexpr uint8 FieldScan8x8[64] = {
-		0,3,8,15,22,30,38,52,
-		1,4,14,21,29,37,45,53,
-		2,7,16,23,31,39,46,58,
-		5,9,20,28,36,44,51,59,
-		6,13,24,32,40,47,54,60,
-		10,17,25,33,41,48,55,61,
-		11,18,26,34,42,49,56,62,
-		12,19,27,35,43,50,57,63
+		0, 3, 8, 15, 22, 30, 38, 52,
+		1, 4, 14, 21, 29, 37, 45, 53,
+		2, 7, 16, 23, 31, 39, 46, 58,
+		5, 9, 20, 28, 36, 44, 51, 59,
+		6, 13, 24, 32, 40, 47, 54, 60,
+		10, 17, 25, 33, 41, 48, 55, 61,
+		11, 18, 26, 34, 42, 49, 56, 62,
+		12, 19, 27, 35, 43, 50, 57, 63
 	};
 
 	inline void ScaleListToWeightScale(const bool bIsFrame, const uint8* ScalingList, const uint8 ListSize, uint8* OutWeightList /* called C in spec*/)
@@ -334,39 +349,39 @@ namespace UE::AVCodecCore::H264
 			}
 		}
 	};
-	
+
 	struct SPS_t : public FNalu
 	{
 	public:
 		U<8, EH264ProfileIDC> profile_idc = EH264ProfileIDC::Auto;
 		U<8, EH264ConstraintFlag> constraint_flags = EH264ConstraintFlag::None;
-		U<8> level_idc; // + profile_idc + constraint_flags = EH264Level
-		UE seq_parameter_set_id; // 0...31 inclusive
-		UE chroma_format_idc = 1; // 0...3 inclusive
-		U<1> separate_colour_plane_flag; // bool
-		UE bit_depth_luma_minus8; // 0...6 inclusive
-		UE bit_depth_chroma_minus8; // 0...6 inclusive
+		U<8> level_idc;							   // + profile_idc + constraint_flags = EH264Level
+		UE seq_parameter_set_id;				   // 0...31 inclusive
+		UE chroma_format_idc = 1;				   // 0...3 inclusive
+		U<1> separate_colour_plane_flag;		   // bool
+		UE bit_depth_luma_minus8;				   // 0...6 inclusive
+		UE bit_depth_chroma_minus8;				   // 0...6 inclusive
 		U<1> qpprime_y_zero_transform_bypass_flag; // 0...1 inclusive
-		U<1> seq_scaling_matrix_present_flag; // bool
-		U<1> seq_scaling_list_present_flag[12]; // bool array
-		uint8 ScalingList4x4[6][16]; // -2^(1+BitDepth)...2^(1+BitDepth)-1 inclusive
-		uint8 ScalingList8x8[6][64]; // -2^(1+BitDepth)...2^(1+BitDepth)-1 inclusive
-		UE log2_max_frame_num_minus4; // 0...12 inclusive
-		UE pic_order_cnt_type; // 0...3 inclusive
-		UE log2_max_pic_order_cnt_lsb_minus4; // 0...12 inclusive
-		U<1> delta_pic_order_always_zero_flag; // bool 
-		SE offset_for_non_ref_pic; // −2^(31)+1...2^(31)−1 inclusive
-		SE offset_for_top_to_bottom_field; // −2^(31)+1...2^(31)−1 inclusive
-		UE num_ref_frames_in_pic_order_cnt_cycle; // 0...255 inclusive
-		SE offset_for_ref_frame[255]; // −2^(31)+1...2^(31)−1 inclusive
-		UE max_num_ref_frames; // 0...MaxDpbFrames inclusive MaxDpbFrames = Min( MaxDecPicBufSizeMacro / ( PicWidthMacro * FrameHeightMacro ), 16) (A.3.1 or A.3.2)
+		U<1> seq_scaling_matrix_present_flag;	   // bool
+		U<1> seq_scaling_list_present_flag[12];	   // bool array
+		uint8 ScalingList4x4[6][16];			   // -2^(1+BitDepth)...2^(1+BitDepth)-1 inclusive
+		uint8 ScalingList8x8[6][64];			   // -2^(1+BitDepth)...2^(1+BitDepth)-1 inclusive
+		UE log2_max_frame_num_minus4;			   // 0...12 inclusive
+		UE pic_order_cnt_type;					   // 0...3 inclusive
+		UE log2_max_pic_order_cnt_lsb_minus4;	   // 0...12 inclusive
+		U<1> delta_pic_order_always_zero_flag;	   // bool
+		SE offset_for_non_ref_pic;				   // −2^(31)+1...2^(31)−1 inclusive
+		SE offset_for_top_to_bottom_field;		   // −2^(31)+1...2^(31)−1 inclusive
+		UE num_ref_frames_in_pic_order_cnt_cycle;  // 0...255 inclusive
+		SE offset_for_ref_frame[255];			   // −2^(31)+1...2^(31)−1 inclusive
+		UE max_num_ref_frames;					   // 0...MaxDpbFrames inclusive MaxDpbFrames = Min( MaxDecPicBufSizeMacro / ( PicWidthMacro * FrameHeightMacro ), 16) (A.3.1 or A.3.2)
 		U<1> gaps_in_frame_num_value_allowed_flag; // bool
-		UE pic_width_in_mbs_minus1; // + 1 = Width in macroblocks
-		UE pic_height_in_map_units_minus1; // + 1 = Height 
-		U<1> frame_mbs_only_flag; // bool
-		U<1> mb_adaptive_frame_field_flag; // bool
-		U<1> direct_8x8_inference_flag; // bool
-		U<1> frame_cropping_flag; // bool 
+		UE pic_width_in_mbs_minus1;				   // + 1 = Width in macroblocks
+		UE pic_height_in_map_units_minus1;		   // + 1 = Height
+		U<1> frame_mbs_only_flag;				   // bool
+		U<1> mb_adaptive_frame_field_flag;		   // bool
+		U<1> direct_8x8_inference_flag;			   // bool
+		U<1> frame_cropping_flag;				   // bool
 		UE frame_crop_left_offset;
 		UE frame_crop_right_offset;
 		UE frame_crop_top_offset;
@@ -375,29 +390,29 @@ namespace UE::AVCodecCore::H264
 		// VUI Parameters
 		U<1> aspect_ratio_info_present_flag; // bool
 		U<8, EH264AspectRatioIDC> aspect_ratio_idc = EH264AspectRatioIDC::Unspecified;
-		U<16> sar_width; // Width in arbitrary unit
-		U<16> sar_height; // Height in same arbitrary unit as Width
-		U<1> overscan_info_present_flag; // bool
-		U<1> overscan_appropriate_flag; // bool
+		U<16> sar_width;					 // Width in arbitrary unit
+		U<16> sar_height;					 // Height in same arbitrary unit as Width
+		U<1> overscan_info_present_flag;	 // bool
+		U<1> overscan_appropriate_flag;		 // bool
 		U<1> video_signal_type_present_flag; // bool
 		U<3, EH264VideoFormat> video_format = EH264VideoFormat::Unspecified;
-		U<1> video_full_range_flag;  // bool
-		U<1> colour_description_present_flag; // bool
-		U<8> colour_primaries = 2; // 0...255 inclusive but see Table E-3 as many values are Reserved
-		U<8> transfer_characteristics = 2; // 0...255 inclusive but see Table E-4 as many values are Reserved
-		U<8> matrix_coefficients = 2; // 0...255 inclusive but see Table E-5 as many values are Reserved
-		U<1> chroma_loc_info_present_flag; // bool
-		UE chroma_sample_loc_type_top_field; // 0...5 inclusive
-		UE chroma_sample_loc_type_bottom_field; // 0...5 inclusive
-		U<1> timing_info_present_flag; // bool
-		U<32> num_units_in_tick = 1; // > 0
-		U<32> time_scale = 1000; // defaulted to ms
-		U<1> fixed_frame_rate_flag; // bool
-		U<1> nal_hrd_parameters_present_flag; // bool
-		U<1> vcl_hrd_parameters_present_flag; // bool
-		U<1> low_delay_hrd_flag; // bool
-		U<1> pic_struct_present_flag; // bool
-		U<1> bitstream_restriction_flag; // bool
+		U<1> video_full_range_flag;					  // bool
+		U<1> colour_description_present_flag;		  // bool
+		U<8> colour_primaries = 2;					  // 0...255 inclusive but see Table E-3 as many values are Reserved
+		U<8> transfer_characteristics = 2;			  // 0...255 inclusive but see Table E-4 as many values are Reserved
+		U<8> matrix_coefficients = 2;				  // 0...255 inclusive but see Table E-5 as many values are Reserved
+		U<1> chroma_loc_info_present_flag;			  // bool
+		UE chroma_sample_loc_type_top_field;		  // 0...5 inclusive
+		UE chroma_sample_loc_type_bottom_field;		  // 0...5 inclusive
+		U<1> timing_info_present_flag;				  // bool
+		U<32> num_units_in_tick = 1;				  // > 0
+		U<32> time_scale = 1000;					  // defaulted to ms
+		U<1> fixed_frame_rate_flag;					  // bool
+		U<1> nal_hrd_parameters_present_flag;		  // bool
+		U<1> vcl_hrd_parameters_present_flag;		  // bool
+		U<1> low_delay_hrd_flag;					  // bool
+		U<1> pic_struct_present_flag;				  // bool
+		U<1> bitstream_restriction_flag;			  // bool
 		U<1> motion_vectors_over_pic_boundaries_flag; // bool
 		UE max_bytes_per_pic_denom;
 		UE max_bits_per_mb_denom;
@@ -409,14 +424,14 @@ namespace UE::AVCodecCore::H264
 		UE cpb_cnt_minus1; // 0...31 inclusive
 		U<4> bit_rate_scale;
 		U<4> cpb_size_scale;
-		UE bit_rate_value_minus1[31]; // 0...2^(32)−2 inclusive TODO (aidan) might need an infer stage to fill out based on profile
-		UE cpb_size_value_minus1[31]; // 0...2^(32)−2 inclusive TODO (aidan) might need an infer stage to fill out based on profile
-		U<1> cbr_flag[31]; // bool
-		U<5> initial_cpb_removal_delay_length_minus1 = 23; // 
+		UE bit_rate_value_minus1[31];					   // 0...2^(32)−2 inclusive TODO (aidan) might need an infer stage to fill out based on profile
+		UE cpb_size_value_minus1[31];					   // 0...2^(32)−2 inclusive TODO (aidan) might need an infer stage to fill out based on profile
+		U<1> cbr_flag[31];								   // bool
+		U<5> initial_cpb_removal_delay_length_minus1 = 23; //
 		U<5> cpb_removal_delay_length_minus1 = 23;
 		U<5> dpb_output_delay_length_minus1 = 23;
 		U<5> time_offset_length = 24;
-		
+
 		EH264Profile GetProfile() const
 		{
 			for (FH264ProfileDefinition const& H264Profile : GH264ProfileDefinitions)
@@ -429,9 +444,8 @@ namespace UE::AVCodecCore::H264
 
 			return EH264Profile::MAX;
 		}
-
 	};
-	
+
 	FAVResult ParseSPS(FBitstreamReader& Bitstream, FNaluInfo const& InNaluInfo, TMap<uint32, SPS_t>& OutMapSPS);
 
 	struct PPS_t : public FNalu
@@ -443,8 +457,8 @@ namespace UE::AVCodecCore::H264
 		UE num_slice_groups_minus1;
 		UE slice_group_map_type;
 		TArray<UE> run_length_minus1; // hard to justify preallocating this as it could be PicWidthInMbs * PicHeightInMapUnits - 1
-		TArray<UE> top_left; // hard to justify preallocating this as it could be PicWidthInMbs * PicHeightInMapUnits - 1
-		TArray<UE> bottom_right; // hard to justify preallocating this as it could be PicWidthInMbs * PicHeightInMapUnits - 1
+		TArray<UE> top_left;		  // hard to justify preallocating this as it could be PicWidthInMbs * PicHeightInMapUnits - 1
+		TArray<UE> bottom_right;	  // hard to justify preallocating this as it could be PicWidthInMbs * PicHeightInMapUnits - 1
 		U<1> slice_group_change_direction_flag;
 		UE slice_group_change_rate_minus1;
 		UE pic_size_in_map_units_minus1;
@@ -466,7 +480,7 @@ namespace UE::AVCodecCore::H264
 		uint8 ScalingList8x8[6][64]; // -2^(1+BitDepth)...2^(1+BitDepth)-1 inclusive
 		SE second_chroma_qp_index_offset;
 	};
-	
+
 	FAVResult ParsePPS(FBitstreamReader& Bitstream, FNaluInfo const& InNaluInfo, TMap<uint32, SPS_t> const& InMapSPS, TMap<uint32, PPS_t>& OutMapPPS);
 
 	struct Slice_t : public FNalu
@@ -495,14 +509,14 @@ namespace UE::AVCodecCore::H264
 		SE slice_alpha_c0_offset_div2;
 		SE slice_beta_offset_div2;
 		U<> slice_group_change_cycle;
-		
+
 		// ref_pic_list_modification
 		struct RefPic_t
 		{
 			U<1> bIsLongTerm;
 			UE pic_num;
 		};
-		
+
 		U<1> ref_pic_list_modification_flag_l0;
 		TArray<RefPic_t> RefPicList0;
 		U<1> ref_pic_list_modification_flag_l1;
@@ -510,23 +524,19 @@ namespace UE::AVCodecCore::H264
 
 		// pred_weight_table
 
-
-		
 		// dec_ref_pic_marking
 		U<1> no_output_of_prior_pic_flag;
 		U<1> long_term_reference_flag;
 
 		U<1> adaptive_ref_pic_marking_mode_flag;
-
-		
 	};
-	
+
 	FAVResult ParseSliceHeader(FBitstreamReader& Bitstream, FNaluInfo const& InNaluInfo, TMap<uint32, SPS_t> InMapSPS, TMap<uint32, PPS_t> const& InMapPPS, Slice_t& OutSlice);
 
 	// SEI can transmit arbitrary data so we probably want something less rigid for this
 	struct SEI_t
 	{
 	};
-	
+
 	FAVResult ParseSEI(FBitstreamReader& Bitstream, FNaluInfo const& InNaluInfo, SEI_t& OutSEI);
 } // namespace UE::AVCodecCore::H264
diff --git a/Engine/Plugins/Experimental/AVCodecs/AVCodecsCore/Source/AVCodecsCore/Public/Video/Encoders/Configs/VideoEncoderConfigH264.h b/Engine/Plugins/Experimental/AVCodecs/AVCodecsCore/Source/AVCodecsCore/Public/Video/Encoders/Configs/VideoEncoderConfigH264.h
--- a/Engine/Plugins/Experimental/AVCodecs/AVCodecsCore/Source/AVCodecsCore/Public/Video/Encoders/Configs/VideoEncoderConfigH264.h
+++ b/Engine/Plugins/Experimental/AVCodecs/AVCodecsCore/Source/AVCodecsCore/Public/Video/Encoders/Configs/VideoEncoderConfigH264.h
@@ -12,6 +12,8 @@ struct FVideoEncoderConfigH264 : public FVideoEncoderConfig
 {
 public:
 	EH264Profile Profile = EH264Profile::Main;
+	EH264AdaptiveTransformMode AdaptiveTransformMode = EH264AdaptiveTransformMode::Auto;
+	EH264EntropyCodingMode EntropyCodingMode = EH264EntropyCodingMode::Auto;
 
 	bool RepeatSPSPPS = false;
 
diff --git a/Engine/Plugins/Experimental/AVCodecs/NVCodecs/Source/NVENC/Private/Video/Encoders/Configs/VideoEncoderConfigNVENC.cpp b/Engine/Plugins/Experimental/AVCodecs/NVCodecs/Source/NVENC/Private/Video/Encoders/Configs/VideoEncoderConfigNVENC.cpp
--- a/Engine/Plugins/Experimental/AVCodecs/NVCodecs/Source/NVENC/Private/Video/Encoders/Configs/VideoEncoderConfigNVENC.cpp
+++ b/Engine/Plugins/Experimental/AVCodecs/NVCodecs/Source/NVENC/Private/Video/Encoders/Configs/VideoEncoderConfigNVENC.cpp
@@ -71,19 +71,19 @@ TAVResult<NV_ENC_BUFFER_FORMAT> FVideoEncoderConfigNVENC::ConvertFormat(EVideoFo
 {
 	switch (Format)
 	{
-	case EVideoFormat::BGRA:
-		return NV_ENC_BUFFER_FORMAT::NV_ENC_BUFFER_FORMAT_ARGB;
-	case EVideoFormat::ABGR10:
-		return NV_ENC_BUFFER_FORMAT::NV_ENC_BUFFER_FORMAT_ARGB10;
-	case EVideoFormat::NV12:
-		return NV_ENC_BUFFER_FORMAT::NV_ENC_BUFFER_FORMAT_NV12;
-	case EVideoFormat::P010:
-		return NV_ENC_BUFFER_FORMAT::NV_ENC_BUFFER_FORMAT_YUV420_10BIT;
-	default:
-		return FAVResult(EAVResult::ErrorUnsupported, FString::Printf(TEXT("Pixel format %d is not supported"), Format), TEXT("NVENC"));
+		case EVideoFormat::BGRA:
+			return NV_ENC_BUFFER_FORMAT::NV_ENC_BUFFER_FORMAT_ARGB;
+		case EVideoFormat::ABGR10:
+			return NV_ENC_BUFFER_FORMAT::NV_ENC_BUFFER_FORMAT_ARGB10;
+		case EVideoFormat::NV12:
+			return NV_ENC_BUFFER_FORMAT::NV_ENC_BUFFER_FORMAT_NV12;
+		case EVideoFormat::P010:
+			return NV_ENC_BUFFER_FORMAT::NV_ENC_BUFFER_FORMAT_YUV420_10BIT;
+		default:
+			return FAVResult(EAVResult::ErrorUnsupported, FString::Printf(TEXT("Pixel format %d is not supported"), Format), TEXT("NVENC"));
 	}
 
-	//NV_ENC_BUFFER_FORMAT::NV_ENC_BUFFER_FORMAT_YUV420_10BIT : NV_ENC_BUFFER_FORMAT::NV_ENC_BUFFER_FORMAT_IYUV;
+	// NV_ENC_BUFFER_FORMAT::NV_ENC_BUFFER_FORMAT_YUV420_10BIT : NV_ENC_BUFFER_FORMAT::NV_ENC_BUFFER_FORMAT_IYUV;
 }
 
 static uint32 DEFAULT_BITRATE_TARGET = 1000000;
@@ -212,9 +212,7 @@ template <>
 DLLEXPORT FAVResult FAVExtension::TransformConfig(FVideoEncoderConfigNVENC& OutConfig, FVideoEncoderConfigH264 const& InConfig)
 {
 	// Unused from preset
-	/*Config.encodeCodecConfig.h264Config.adaptiveTransformMode = NV_ENC_H264_ADAPTIVE_TRANSFORM_ENABLE;
-	Config.encodeCodecConfig.h264Config.entropyCodingMode = NV_ENC_H264_ENTROPY_CODING_MODE_CABAC;
-	Config.encodeCodecConfig.h264Config.h264VUIParameters.videoFormat = 5;
+	/*Config.encodeCodecConfig.h264Config.h264VUIParameters.videoFormat = 5;
 	Config.encodeCodecConfig.h264Config.h264VUIParameters.colourPrimaries = 2;
 	Config.encodeCodecConfig.h264Config.h264VUIParameters.transferCharacteristics = 2;
 	Config.encodeCodecConfig.h264Config.h264VUIParameters.colourMatrix = 2;*/
@@ -260,6 +258,50 @@ DLLEXPORT FAVResult FAVExtension::TransformConfig(FVideoEncoderConfigNVENC& OutC
 	OutH264Config.chromaFormatIDC = 1;
 	OutH264Config.enableFillerDataInsertion = InConfig.bFillData ? 1 : 0;
 
+	static auto const ConvertAdaptiveTransformMode = [](EH264AdaptiveTransformMode TransformMode) -> TAVResult<NV_ENC_H264_ADAPTIVE_TRANSFORM_MODE> {
+		switch (TransformMode)
+		{
+			case EH264AdaptiveTransformMode::Auto:
+				return NV_ENC_H264_ADAPTIVE_TRANSFORM_AUTOSELECT;
+			case EH264AdaptiveTransformMode::Disable:
+				return NV_ENC_H264_ADAPTIVE_TRANSFORM_DISABLE;
+			case EH264AdaptiveTransformMode::Enable:
+				return NV_ENC_H264_ADAPTIVE_TRANSFORM_ENABLE;
+			default:
+				return FAVResult(EAVResult::ErrorUnsupported, FString::Printf(TEXT("H264 transform mode %d is not supported"), TransformMode), TEXT("NVENC"));
+		}
+	};
+
+	TAVResult<NV_ENC_H264_ADAPTIVE_TRANSFORM_MODE> const ConvertedAdaptiveTransformMode = ConvertAdaptiveTransformMode(InConfig.AdaptiveTransformMode);
+	if (ConvertedAdaptiveTransformMode.IsNotSuccess())
+	{
+		return ConvertedAdaptiveTransformMode;
+	}
+
+	OutH264Config.adaptiveTransformMode = ConvertedAdaptiveTransformMode;
+
+	static auto const ConvertEntropyCodingMode = [](EH264EntropyCodingMode EntropyCodingMode) -> TAVResult<NV_ENC_H264_ENTROPY_CODING_MODE> {
+		switch (EntropyCodingMode)
+		{
+			case EH264EntropyCodingMode::Auto:
+				return NV_ENC_H264_ENTROPY_CODING_MODE_AUTOSELECT;
+			case EH264EntropyCodingMode::CABAC:
+				return NV_ENC_H264_ENTROPY_CODING_MODE_CABAC;
+			case EH264EntropyCodingMode::CAVLC:
+				return NV_ENC_H264_ENTROPY_CODING_MODE_CAVLC;
+			default:
+				return FAVResult(EAVResult::ErrorUnsupported, FString::Printf(TEXT("H264 entropy coding mode %d is not supported"), EntropyCodingMode), TEXT("NVENC"));
+		}
+	};
+
+	TAVResult<NV_ENC_H264_ENTROPY_CODING_MODE> const ConvertedEntropyCodingMode = ConvertEntropyCodingMode(InConfig.EntropyCodingMode);
+	if (ConvertedEntropyCodingMode.IsNotSuccess())
+	{
+		return ConvertedEntropyCodingMode;
+	}
+
+	OutH264Config.entropyCodingMode = ConvertedEntropyCodingMode;
+
 	if (InConfig.RateControlMode == ERateControlMode::CBR && InConfig.bFillData)
 	{
 		// `outputPictureTimingSEI` is used in CBR mode to fill video frame with data to match the requested bitrate.
diff --git a/Engine/Plugins/Media/PixelStreaming/Source/PixelStreaming/Private/VideoEncoderSingleLayerHardware.cpp b/Engine/Plugins/Media/PixelStreaming/Source/PixelStreaming/Private/VideoEncoderSingleLayerHardware.cpp
--- a/Engine/Plugins/Media/PixelStreaming/Source/PixelStreaming/Private/VideoEncoderSingleLayerHardware.cpp
+++ b/Engine/Plugins/Media/PixelStreaming/Source/PixelStreaming/Private/VideoEncoderSingleLayerHardware.cpp
@@ -58,6 +58,10 @@ namespace UE::PixelStreaming
 				VideoConfig->IntraRefreshPeriodFrames = PixelStreaming::Settings::CVarPixelStreamingEncoderIntraRefreshPeriodFrames.GetValueOnAnyThread();
 				VideoConfig->IntraRefreshCountFrames = PixelStreaming::Settings::CVarPixelStreamingEncoderIntraRefreshCountFrames.GetValueOnAnyThread();
 				VideoConfig->KeyframeInterval = PixelStreaming::Settings::CVarPixelStreamingEncoderKeyframeInterval.GetValueOnAnyThread();
+				// The WebRTC spec can only guarantee that the Baseline profile is supported. Therefore we use Baseline, but enable these extra
+				// features to improve bitrate usage
+				VideoConfig->AdaptiveTransformMode = EH264AdaptiveTransformMode::Enable;
+				VideoConfig->EntropyCodingMode = EH264EntropyCodingMode::CABAC;
 				InitialVideoConfig = MoveTemp(VideoConfig);
 				return WEBRTC_VIDEO_CODEC_OK;
 			}
@@ -175,7 +179,7 @@ namespace UE::PixelStreaming
 		if (TSharedPtr<FVideoEncoderHardware> const& PinnedHardwareEncoder = HardwareEncoder.Pin())
 		{
 			TRACE_CPUPROFILER_EVENT_SCOPE_ON_CHANNEL_STR("PixelStreaming Hardware Encoding", PixelStreamingChannel);
-			
+
 			IPixelCaptureOutputFrame* AdaptedLayer = FrameBuffer->RequestFormat(PixelCaptureBufferFormat::FORMAT_RHI);
 
 			if (AdaptedLayer == nullptr)
