Replace all usage of MinWindows.h with WindowsHWrapper for consistency.

 - only a handful of places so there's a clear precedent

#rnx

diff --git a/Engine/Plugins/Experimental/AVCodecs/AVCodecsCore/Source/AVCodecsCoreRHI/Public/Video/Resources/VideoResourceRHI.h b/Engine/Plugins/Experimental/AVCodecs/AVCodecsCore/Source/AVCodecsCoreRHI/Public/Video/Resources/VideoResourceRHI.h
--- a/Engine/Plugins/Experimental/AVCodecs/AVCodecsCore/Source/AVCodecsCoreRHI/Public/Video/Resources/VideoResourceRHI.h
+++ b/Engine/Plugins/Experimental/AVCodecs/AVCodecsCore/Source/AVCodecsCoreRHI/Public/Video/Resources/VideoResourceRHI.h
@@ -9,7 +9,7 @@
 
 #if PLATFORM_WINDOWS
 THIRD_PARTY_INCLUDES_START
-#include "Windows/MinWindows.h"
+#include "Windows/WindowsHWrapper.h"
 #include "d3d12.h"
 THIRD_PARTY_INCLUDES_END
 #endif
diff --git a/Engine/Source/Runtime/Core/Private/AutoRTFM/StdLibSupport.cpp b/Engine/Source/Runtime/Core/Private/AutoRTFM/StdLibSupport.cpp
--- a/Engine/Source/Runtime/Core/Private/AutoRTFM/StdLibSupport.cpp
+++ b/Engine/Source/Runtime/Core/Private/AutoRTFM/StdLibSupport.cpp
@@ -17,9 +17,7 @@
 #endif
 
 #if PLATFORM_WINDOWS
-#include "Windows/PreWindowsApi.h"
-#include "Windows/MinWindows.h"
-#include "Windows/PostWindowsApi.h"
+#include "Windows/WindowsHWrapper.h"
 #endif
 
 namespace AutoRTFM
diff --git a/Engine/Source/Runtime/Core/Private/AutoRTFM/Utils.cpp b/Engine/Source/Runtime/Core/Private/AutoRTFM/Utils.cpp
--- a/Engine/Source/Runtime/Core/Private/AutoRTFM/Utils.cpp
+++ b/Engine/Source/Runtime/Core/Private/AutoRTFM/Utils.cpp
@@ -7,9 +7,7 @@
 #include <string.h>
 
 #if PLATFORM_WINDOWS
-#include "Windows/PreWindowsApi.h"
-#include "Windows/MinWindows.h"
-#include "Windows/PostWindowsApi.h"
+#include "Windows/WindowsHWrapper.h"
 #include <dbghelp.h>
 #else
 #include <execinfo.h>
