Add: VideoToolbox Encoder and Decoders to AVCodecs

diff --git a/Engine/Plugins/Experimental/AVCodecs/AMFCodecs/Source/AMFCodecs/Private/AMFCodecsModule.cpp b/Engine/Plugins/Experimental/AVCodecs/AMFCodecs/Source/AMFCodecs/Private/AMFCodecsModule.cpp
--- a/Engine/Plugins/Experimental/AVCodecs/AMFCodecs/Source/AMFCodecs/Private/AMFCodecsModule.cpp
+++ b/Engine/Plugins/Experimental/AVCodecs/AMFCodecs/Source/AMFCodecs/Private/AMFCodecsModule.cpp
@@ -7,10 +7,10 @@
 #include "Video/Encoders/Configs/VideoEncoderConfigH264.h"
 #include "Video/Encoders/Configs/VideoEncoderConfigH265.h"
 
-#include "Video/Resources/VideoResourceVulkan.h"
+#include "Video/Resources/Vulkan/VideoResourceVulkan.h"
 
 #if PLATFORM_WINDOWS
-#include "Video/Resources/Windows/VideoResourceD3D.h"
+#include "Video/Resources/D3D/VideoResourceD3D.h"
 #endif
 
 #include "Video/Decoders/VideoDecoderAMF.h"
@@ -36,6 +36,8 @@ public:
 					return FAPI::Get<FAMF>().IsValid();
 				});
 
+		// TODO (william.belcher): H265 decoding doesn't currently work with AMD
+		/* 
 		FVideoDecoder::RegisterPermutationsOf<TVideoDecoderAMF<FVideoResourceD3D11>>
 			::With<FVideoResourceD3D11>
 			::And<FVideoDecoderConfigAMF, FVideoDecoderConfigH264, FVideoDecoderConfigH265>([](TSharedRef<FAVDevice> const& NewDevice, TSharedRef<FAVInstance> const& NewInstance)
@@ -49,6 +51,7 @@ public:
 				{
 					return FAPI::Get<FAMF>().IsValid();
 				});
+		*/
 #endif
 
 		/*FVideoEncoder::RegisterPermutationsOf<TVideoEncoderAMF<FVideoResourceVulkan>>
diff --git a/Engine/Plugins/Experimental/AVCodecs/AMFCodecs/Source/AMFCodecs/Private/Video/Decoders/Configs/VideoDecoderConfigAMF.cpp b/Engine/Plugins/Experimental/AVCodecs/AMFCodecs/Source/AMFCodecs/Private/Video/Decoders/Configs/VideoDecoderConfigAMF.cpp
--- a/Engine/Plugins/Experimental/AVCodecs/AMFCodecs/Source/AMFCodecs/Private/Video/Decoders/Configs/VideoDecoderConfigAMF.cpp
+++ b/Engine/Plugins/Experimental/AVCodecs/AMFCodecs/Source/AMFCodecs/Private/Video/Decoders/Configs/VideoDecoderConfigAMF.cpp
@@ -24,12 +24,12 @@ FAVResult FVideoDecoderConfigAMF::Parse(TSharedRef<FAVInstance> const& Instance,
 
 		FVideoDecoderConfigH265& H265 = Instance->Edit<FVideoDecoderConfigH265>();
 
-		TArray<TSharedPtr<FNaluH265>> Nalus;
-		FAVResult const Result = H265.Parse(Instance, Packet, Nalus);
-		if (Result.IsNotSuccess())
-		{
-			return FAVResult(EAVResult::ErrorResolving, TEXT("Failed to parse H265 Bitstream"));
-		}
+		// TArray<TSharedPtr<FNaluH265>> Nalus;
+		// FAVResult const Result = H265.Parse(Packet, Nalus);
+		// if (Result.IsNotSuccess())
+		// {
+		// 	return FAVResult(EAVResult::ErrorResolving, TEXT("Failed to parse H265 Bitstream"));
+		// }
 
 		// TODO (Aidan) Write SPS+PPS in contiguous block to each FParsedPicture in OutPictures
 
diff --git a/Engine/Plugins/Experimental/AVCodecs/AMFCodecs/Source/AMFCodecs/Private/Video/Decoders/VideoDecoderAMF.cpp b/Engine/Plugins/Experimental/AVCodecs/AMFCodecs/Source/AMFCodecs/Private/Video/Decoders/VideoDecoderAMF.cpp
--- a/Engine/Plugins/Experimental/AVCodecs/AMFCodecs/Source/AMFCodecs/Private/Video/Decoders/VideoDecoderAMF.cpp
+++ b/Engine/Plugins/Experimental/AVCodecs/AMFCodecs/Source/AMFCodecs/Private/Video/Decoders/VideoDecoderAMF.cpp
@@ -2,10 +2,10 @@
 
 #include "Video/Decoders/VideoDecoderAMF.h"
 
-#include "Video/Resources/VideoResourceVulkan.h"
+#include "Video/Resources/Vulkan/VideoResourceVulkan.h"
 
 #if PLATFORM_WINDOWS
-#include "Video/Resources/Windows/VideoResourceD3D.h"
+#include "Video/Resources/D3D/VideoResourceD3D.h"
 #endif
 
 #if PLATFORM_WINDOWS
diff --git a/Engine/Plugins/Experimental/AVCodecs/AMFCodecs/Source/AMFCodecs/Private/Video/Encoders/VideoEncoderAMF.cpp b/Engine/Plugins/Experimental/AVCodecs/AMFCodecs/Source/AMFCodecs/Private/Video/Encoders/VideoEncoderAMF.cpp
--- a/Engine/Plugins/Experimental/AVCodecs/AMFCodecs/Source/AMFCodecs/Private/Video/Encoders/VideoEncoderAMF.cpp
+++ b/Engine/Plugins/Experimental/AVCodecs/AMFCodecs/Source/AMFCodecs/Private/Video/Encoders/VideoEncoderAMF.cpp
@@ -2,10 +2,10 @@
 
 #include "Video/Encoders/VideoEncoderAMF.h"
 
-#include "Video/Resources/VideoResourceVulkan.h"
+#include "Video/Resources/Vulkan/VideoResourceVulkan.h"
 
 #if PLATFORM_WINDOWS
-	#include "Video/Resources/Windows/VideoResourceD3D.h"
+	#include "Video/Resources/D3D/VideoResourceD3D.h"
 #endif
 
 REGISTER_TYPEID(amf::AMFSurfacePtr);
diff --git a/Engine/Plugins/Experimental/AVCodecs/AVCodecsCore/AVCodecsCore.uplugin b/Engine/Plugins/Experimental/AVCodecs/AVCodecsCore/AVCodecsCore.uplugin
--- a/Engine/Plugins/Experimental/AVCodecs/AVCodecsCore/AVCodecsCore.uplugin
+++ b/Engine/Plugins/Experimental/AVCodecs/AVCodecsCore/AVCodecsCore.uplugin
@@ -18,7 +18,7 @@
 			"Name": "AVCodecsCore",
 			"Type": "Runtime",
 			"LoadingPhase": "PostEngineInit",
-			"PlatformAllowList": [ "Win64",	"Linux" ]
+			"PlatformAllowList": [ "Win64",	"Linux", "Mac" ]
 		},
 		{
 			"Name": "AVCodecsCoreRHI",
@@ -27,11 +27,12 @@
 			"TargetDenyList": [
 				"Program"
 			],
-			"PlatformAllowList": [ "Win64",	"Linux"	]
+			"PlatformAllowList": [ "Win64",	"Linux", "Mac"	]
 		}
 	],
 	"SupportedTargetPlatforms": [
 		"Win64",
-		"Linux"
+		"Linux",
+		"Mac"
 	]
 }
\ No newline at end of file
diff --git a/Engine/Plugins/Experimental/AVCodecs/AVCodecsCore/Source/AVCodecsCore/AVCodecsCore.Build.cs b/Engine/Plugins/Experimental/AVCodecs/AVCodecsCore/Source/AVCodecsCore/AVCodecsCore.Build.cs
--- a/Engine/Plugins/Experimental/AVCodecs/AVCodecsCore/Source/AVCodecsCore/AVCodecsCore.Build.cs
+++ b/Engine/Plugins/Experimental/AVCodecs/AVCodecsCore/Source/AVCodecsCore/AVCodecsCore.Build.cs
@@ -4,55 +4,88 @@ using UnrealBuildTool;
 using System.IO;
 using System;
 
-[SupportedPlatforms("Win64", "Linux")]
+[SupportedPlatforms("Win64", "Linux", "Mac")]
 public class AVCodecsCore : ModuleRules
 {
-	public AVCodecsCore(ReadOnlyTargetRules Target) : base(Target)
-	{
-		// Without these two compilation fails on VS2017 with D8049: command line is too long to fit in debug record.
-		bLegacyPublicIncludePaths = false;
-		DefaultBuildSettings = BuildSettingsVersion.V2;
+    public AVCodecsCore(ReadOnlyTargetRules Target) : base(Target)
+    {
+        // Without these two compilation fails on VS2017 with D8049: command line is too long to fit in debug record.
+        bLegacyPublicIncludePaths = false;
+        DefaultBuildSettings = BuildSettingsVersion.V2;
 
-		// PCHUsage = PCHUsageMode.NoPCHs;
+        // PCHUsage = PCHUsageMode.NoPCHs;
 
-		// PrecompileForTargets = PrecompileTargetsType.None;
+        // PrecompileForTargets = PrecompileTargetsType.None;
 
-		PublicIncludePaths.AddRange(new string[] {
+        PublicIncludePaths.AddRange(new string[] {
 			// ... add public include paths required here ...
 		});
 
-		PrivateIncludePaths.AddRange(new string[] {
+        PrivateIncludePaths.AddRange(new string[] {
 			// ... add other private include paths required here ...
 		});
 
-		PrivateDependencyModuleNames.AddRange(new string[] {
-			"Engine"
-		});
+        PrivateDependencyModuleNames.AddRange(new string[] {
+            "Engine"
+        });
 
-		PublicDependencyModuleNames.AddRange(new string[] {
-			"Core",
+        PublicDependencyModuleNames.AddRange(new string[] {
+            "Core",
 			// ... add other public dependencies that you statically link with here ...
 		});
 
-		DynamicallyLoadedModuleNames.AddRange(new string[] {
+        DynamicallyLoadedModuleNames.AddRange(new string[] {
 			// ... add any modules that your module loads dynamically here ...
 		});
 
-		if (Target.IsInPlatformGroup(UnrealPlatformGroup.Windows) || Target.IsInPlatformGroup(UnrealPlatformGroup.Unix))
+        if (Target.IsInPlatformGroup(UnrealPlatformGroup.Windows) || Target.IsInPlatformGroup(UnrealPlatformGroup.Unix))
+        {
+            PublicIncludePathModuleNames.Add("Vulkan");
+            AddEngineThirdPartyPrivateStaticDependencies(Target, "Vulkan");
+
+        	PublicDefinitions.Add("AVCODECS_USE_VULKAN=1");
+		} 
+		else
+		{
+			PublicDefinitions.Add("AVCODECS_USE_VULKAN=0");
+		}
+
+        if (Target.IsInPlatformGroup(UnrealPlatformGroup.Windows))
+        {
+            AddEngineThirdPartyPrivateStaticDependencies(Target, "DX11", "DX12");
+
+            PublicDelayLoadDLLs.Add("mfplat.dll");
+            PublicDelayLoadDLLs.Add("mfuuid.dll");
+            PublicDelayLoadDLLs.Add("Mfreadwrite.dll");
+
+        	PublicDefinitions.Add("AVCODECS_USE_D3D=1");
+		}
+		else
 		{
-			PublicIncludePathModuleNames.Add("Vulkan");
-			AddEngineThirdPartyPrivateStaticDependencies(Target, "Vulkan");
+			PublicDefinitions.Add("AVCODECS_USE_D3D=0");
 		}
 
-		if (Target.IsInPlatformGroup(UnrealPlatformGroup.Windows))
+        if (Target.IsInPlatformGroup(UnrealPlatformGroup.Apple))
 		{
-			AddEngineThirdPartyPrivateStaticDependencies(Target, "DX11", "DX12");
+			PublicFrameworks.AddRange(new string[]{
+				"AVFoundation",
+				"VideoToolbox"
+			});
+   
+            PublicDependencyModuleNames.AddRange(new string[] {
+                "MetalRHI",
+            });
+   
+            PublicWeakFrameworks.Add("Metal");
+			AddEngineThirdPartyPrivateStaticDependencies(Target, "MTLPP");
 			
-			PublicDelayLoadDLLs.Add("mfplat.dll");
-			PublicDelayLoadDLLs.Add("mfuuid.dll");
-			PublicDelayLoadDLLs.Add("Mfreadwrite.dll");
+			PublicDefinitions.Add("AVCODECS_USE_METAL=1");
+		}
+		else
+		{
+			PublicDefinitions.Add("AVCODECS_USE_METAL=0");
 		}
 
-		PublicDefinitions.Add("DEBUG_DUMP_TO_DISK=0");
-	}
+        PublicDefinitions.Add("DEBUG_DUMP_TO_DISK=0");
+    }
 }
diff --git a/Engine/Plugins/Experimental/AVCodecs/AVCodecsCore/Source/AVCodecsCore/Private/Video/CodecUtils/CodecUtilsH264.cpp b/Engine/Plugins/Experimental/AVCodecs/AVCodecsCore/Source/AVCodecsCore/Private/Video/CodecUtils/CodecUtilsH264.cpp
--- a/Engine/Plugins/Experimental/AVCodecs/AVCodecsCore/Source/AVCodecsCore/Private/Video/CodecUtils/CodecUtilsH264.cpp
+++ b/Engine/Plugins/Experimental/AVCodecs/AVCodecsCore/Source/AVCodecsCore/Private/Video/CodecUtils/CodecUtilsH264.cpp
@@ -4,7 +4,6 @@
 
 #include "Containers/Array.h"
 #include "AVResult.h"
-#include "Video/CodecUtils/CodecUtilsH265.h"
 
 FH264ProfileDefinition GH264ProfileDefinitions[static_cast<uint8>(EH264Profile::MAX)] = {
 	{ EH264Profile::Auto, UE::AVCodecCore::H264::EH264ProfileIDC::Auto, UE::AVCodecCore::H264::EH264ConstraintFlag::None, TEXT("Auto") },
@@ -34,12 +33,12 @@ FH264ProfileDefinition GH264ProfileDefinitions[static_cast<uint8>(EH264Profile::
 
 namespace UE::AVCodecCore::H264
 {
-	FAVResult FindNALUs(FVideoPacket const& InPacket, TArray<FNaluInfo>& FoundNalus)
+	FAVResult FindNALUs(FVideoPacket const& InPacket, TArray<FNaluH264>& FoundNalus)
 	{
 		FBitstreamReader Reader(InPacket.DataPtr.Get(), InPacket.DataSize);
 		if (Reader.NumBytesRemaining() < 3)
 		{
-			return FAVResult(EAVResult::Warning, TEXT("Bitstream not long enough to hold a NALU"), TEXT("H264"));
+			return FAVResult(EAVResult::Warning, TEXT("Bitstream not long enough to hold a NALU"), TEXT("BitstreamParserH264"));
 		}
 
 		TArrayView64<uint8> const Data = InPacket.GetData();
@@ -56,7 +55,7 @@ namespace UE::AVCodecCore::H264
 				if (Data[i + 1] == 0 && Data[i] == 0)
 				{
 					// Found start sequence of NalU but we don't know if it has a 3 or 4 byte start code so we check.
-					FNaluInfo NalInfo = { (uint64)i, 0, 3, 0, ENaluType::Unspecified, nullptr };
+					FNaluH264 NalInfo = { (uint64)i, 0, 3, 0, ENaluType::Unspecified, nullptr };
 					if (NalInfo.Start > 0 && Data[NalInfo.Start - 1] == 0)
 					{
 						++NalInfo.StartCodeSize;
@@ -94,13 +93,13 @@ namespace UE::AVCodecCore::H264
 		}
 
 		// Last Nal size is the remaining size of the bitstream minus a trailing zero byte
-		FoundNalus.Last().Size = Data.Num() - (FoundNalus.Last().Start + FoundNalus.Last().StartCodeSize + 1);
+		FoundNalus.Last().Size = Data.Num() - (FoundNalus.Last().Start + FoundNalus.Last().StartCodeSize);
 
 	#if !IS_PROGRAM
 
 		FAVResult::Log(EAVResult::Success, FString::Printf(TEXT("FAVExtension::TransformConfig found %d NALUs in bitdatastream"), FoundNalus.Num()));
 
-		for (const FNaluInfo& NalUInfo : FoundNalus)
+		for (const FNaluH264& NalUInfo : FoundNalus)
 		{
 			FAVResult::Log(EAVResult::Success, FString::Printf(TEXT("Found NALU at %llu size %llu with type %u"), NalUInfo.Start, NalUInfo.Size, (uint8)NalUInfo.Type));
 		}
@@ -110,7 +109,7 @@ namespace UE::AVCodecCore::H264
 		return EAVResult::Success;
 	}
 
-	FAVResult ParseSEI(FBitstreamReader& Bitstream, FNaluInfo const& InNaluInfo, SEI_t& OutSEI)
+	FAVResult ParseSEI(FBitstreamReader& Bitstream, FNaluH264 const& InNaluInfo, SEI_t& OutSEI)
 	{
 		//unimplemented();
 		return FAVResult(EAVResult::Success, TEXT("SEI Unimplemented"));
@@ -177,28 +176,28 @@ namespace UE::AVCodecCore::H264
 		}
 	}
 
-	FAVResult ParseSPS(FBitstreamReader& Bitstream, FNaluInfo const& InNaluInfo, TMap<uint32, SPS_t>& OutMapSPS)
+	FAVResult ParseSPS(FBitstreamReader& Bitstream, FNaluH264 const& InNaluInfo, TMap<uint32, SPS_t>& OutMapSPS)
 	{
 		// Be sure to maintain parity with SPS struct values, or we may miss bits
-		FNalu::U<8, EH264ProfileIDC> temp_profile_idc;
-		FNalu::U<8, EH264ConstraintFlag> temp_constraint_flags;
+		FNalu::U<8, EH264ProfileIDC> profile_idc;
+        Bitstream.Read(profile_idc); // u(8)
 
-		Bitstream.Read(
-			temp_profile_idc,
-			temp_constraint_flags);
+		FNalu::U<8, EH264ConstraintFlag> constraint_flags;
+        Bitstream.Read(constraint_flags); // u(8)
+
+        FNalu::U<8> level_idc;
+		Bitstream.Read(level_idc); // u(8)
 		
-		FNalu::SE seq_parameter_set_id;
-		Bitstream.Read(seq_parameter_set_id);
+		FNalu::UE seq_parameter_set_id;
+		Bitstream.Read(seq_parameter_set_id); // ue(v)
 
 		// Find SPS at current ID or set it to the map
 		SPS_t& OutSPS = OutMapSPS.FindOrAdd(seq_parameter_set_id);
 
-		OutSPS.profile_idc = temp_profile_idc;
-		OutSPS.constraint_flags = temp_constraint_flags;
-		
-		Bitstream.Read(
-			OutSPS.level_idc,
-			OutSPS.seq_parameter_set_id);
+        OutSPS.profile_idc = profile_idc;
+		OutSPS.constraint_flags = constraint_flags;
+        OutSPS.level_idc = level_idc;
+        OutSPS.seq_parameter_set_id = seq_parameter_set_id;
 
 		if (OutSPS.profile_idc == EH264ProfileIDC::High ||
 			OutSPS.profile_idc == EH264ProfileIDC::High10 ||
@@ -208,20 +207,23 @@ namespace UE::AVCodecCore::H264
 			OutSPS.profile_idc == EH264ProfileIDC::ScalableBaseline ||
 			OutSPS.profile_idc == EH264ProfileIDC::ScalableHigh ||
 			OutSPS.profile_idc == EH264ProfileIDC::MultiviewHigh ||
-			OutSPS.profile_idc == EH264ProfileIDC::MultiviewDepthHigh)
+            OutSPS.profile_idc == EH264ProfileIDC::StereoHigh ||
+			OutSPS.profile_idc == EH264ProfileIDC::MultiviewDepthHigh ||
+            OutSPS.profile_idc == EH264ProfileIDC::MultiresolutionFrameCompatibleHigh ||
+            OutSPS.profile_idc == EH264ProfileIDC::EnhancedMultiviewDepthHigh)
 		{
-			Bitstream.Read(OutSPS.chroma_format_idc);
+			Bitstream.Read(OutSPS.chroma_format_idc); // ue(v)
 
 			if (OutSPS.chroma_format_idc == 3)
 			{
-				Bitstream.Read(OutSPS.separate_colour_plane_flag);
+				Bitstream.Read(OutSPS.separate_colour_plane_flag); // u(1)
 			}
 
 			Bitstream.Read(
-				OutSPS.bit_depth_luma_minus8,
-				OutSPS.bit_depth_chroma_minus8,
-				OutSPS.qpprime_y_zero_transform_bypass_flag,
-				OutSPS.seq_scaling_matrix_present_flag);
+				OutSPS.bit_depth_luma_minus8,                // ue(v)
+				OutSPS.bit_depth_chroma_minus8,              // ue(8)
+				OutSPS.qpprime_y_zero_transform_bypass_flag, // u(1)
+				OutSPS.seq_scaling_matrix_present_flag);     // u(1)
 
 			if (OutSPS.seq_scaling_matrix_present_flag)
 			{
@@ -229,187 +231,97 @@ namespace UE::AVCodecCore::H264
 			}
 		}
 
-		Bitstream.Read(OutSPS.log2_max_frame_num_minus4);
+        const uint32_t MaxLog2Minus4 = 12;
+		Bitstream.Read(OutSPS.log2_max_frame_num_minus4); // ue(v)
+        if(OutSPS.log2_max_frame_num_minus4 > MaxLog2Minus4)
+        {
+			return FAVResult(EAVResult::Error, TEXT("log2_max_frame_num_minus4 > MaxLog2Minus4"), TEXT("H264"));
+        }
 
-		Bitstream.Read(OutSPS.pic_order_cnt_type);
+		Bitstream.Read(OutSPS.pic_order_cnt_type); // ue(v)
 		if (OutSPS.pic_order_cnt_type == 0)
 		{
-			Bitstream.Read(OutSPS.log2_max_pic_order_cnt_lsb_minus4);
+			Bitstream.Read(OutSPS.log2_max_pic_order_cnt_lsb_minus4); // ue(v)
+            if(OutSPS.log2_max_pic_order_cnt_lsb_minus4 > MaxLog2Minus4)
+            {
+                return FAVResult(EAVResult::Error, TEXT("log2_max_pic_order_cnt_lsb_minus4 > MaxLog2Minus4"), TEXT("H264"));
+            }
 		}
 		else if (OutSPS.pic_order_cnt_type == 1)
 		{
 			Bitstream.Read(
-				OutSPS.delta_pic_order_always_zero_flag,
-				OutSPS.offset_for_non_ref_pic,
-				OutSPS.offset_for_top_to_bottom_field,
-				OutSPS.num_ref_frames_in_pic_order_cnt_cycle);
+				OutSPS.delta_pic_order_always_zero_flag,       // u(1)
+				OutSPS.offset_for_non_ref_pic,                 // se(v)
+				OutSPS.offset_for_top_to_bottom_field,         // se(v)
+				OutSPS.num_ref_frames_in_pic_order_cnt_cycle); // ue(v)
 			
 			for (uint32 i = 0; i < OutSPS.num_ref_frames_in_pic_order_cnt_cycle; ++i)
 			{
-				Bitstream.Read(OutSPS.offset_for_ref_frame[i]);
+				Bitstream.Read(OutSPS.offset_for_ref_frame[i]); // se(v)
 			}
 		}
 
 		Bitstream.Read(
-			OutSPS.max_num_ref_frames,
-			OutSPS.gaps_in_frame_num_value_allowed_flag,
-			OutSPS.pic_width_in_mbs_minus1,
-			OutSPS.pic_height_in_map_units_minus1);
+			OutSPS.max_num_ref_frames,                   // ue(v)
+			OutSPS.gaps_in_frame_num_value_allowed_flag, // u(1)
+			OutSPS.pic_width_in_mbs_minus1,              // ue(v)
+			OutSPS.pic_height_in_map_units_minus1);      // ue(v)
 
-		Bitstream.Read(OutSPS.frame_mbs_only_flag);
+		Bitstream.Read(OutSPS.frame_mbs_only_flag); // u(1)
 		if (!OutSPS.frame_mbs_only_flag)
 		{
-			Bitstream.Read(OutSPS.mb_adaptive_frame_field_flag);
+			Bitstream.Read(OutSPS.mb_adaptive_frame_field_flag); // u(1)
 		}
 
-		Bitstream.Read(OutSPS.direct_8x8_inference_flag);
+		Bitstream.Read(OutSPS.direct_8x8_inference_flag); // u(1)
 
-		Bitstream.Read(OutSPS.frame_cropping_flag);
+		Bitstream.Read(OutSPS.frame_cropping_flag); // u(1)
 		if (OutSPS.frame_cropping_flag)
 		{
 			Bitstream.Read(
-				OutSPS.frame_crop_left_offset,
-				OutSPS.frame_crop_right_offset,
-				OutSPS.frame_crop_top_offset,
-				OutSPS.frame_crop_bottom_offset);
+				OutSPS.frame_crop_left_offset,    // ue(v)
+				OutSPS.frame_crop_right_offset,   // ue(v)
+				OutSPS.frame_crop_top_offset,     // ue(v)
+				OutSPS.frame_crop_bottom_offset); // ue(v)
 		}
 
-		Bitstream.Read(OutSPS.vui_parameters_present_flag);
-		if (OutSPS.vui_parameters_present_flag)
-		{
-			Bitstream.Read(OutSPS.aspect_ratio_info_present_flag);
-			if (OutSPS.aspect_ratio_info_present_flag)
-			{
-				OutSPS.aspect_ratio_idc = (EH264AspectRatioIDC)Bitstream.ReadBits(8);
-				if (OutSPS.aspect_ratio_idc == EH264AspectRatioIDC::Extended_SAR)
-				{
-					Bitstream.Read(OutSPS.sar_width);
-					Bitstream.Read(OutSPS.sar_height);
-				}
-			}
-
-			Bitstream.Read(OutSPS.overscan_info_present_flag);
-			if (OutSPS.overscan_info_present_flag)
-			{
-				Bitstream.Read(OutSPS.overscan_appropriate_flag);
-			}
-
-			Bitstream.Read(OutSPS.video_signal_type_present_flag);
-			if (OutSPS.video_signal_type_present_flag)
-			{
-				Bitstream.Read(
-					OutSPS.video_format,
-					OutSPS.video_full_range_flag);
-
-				Bitstream.Read(OutSPS.colour_description_present_flag);
-				if (OutSPS.colour_description_present_flag)
-				{
-					Bitstream.Read(
-						OutSPS.colour_primaries,
-						OutSPS.transfer_characteristics,
-						OutSPS.matrix_coefficients);
-				}
-			}
-
-			Bitstream.Read(OutSPS.chroma_loc_info_present_flag);
-			if (OutSPS.chroma_loc_info_present_flag)
-			{
-				Bitstream.Read(
-					OutSPS.chroma_sample_loc_type_top_field,
-					OutSPS.chroma_sample_loc_type_bottom_field);
-			}
-			
-			Bitstream.Read(OutSPS.timing_info_present_flag);
-			if (OutSPS.timing_info_present_flag)
-			{
-				Bitstream.Read(
-					OutSPS.num_units_in_tick,
-					OutSPS.time_scale,
-					OutSPS.fixed_frame_rate_flag);
-			}
-
-			auto hrd_parameters = [&Bitstream, &OutSPS]() -> void
-			{
-				Bitstream.Read(
-					OutSPS.cpb_cnt_minus1,
-					OutSPS.bit_rate_scale,
-					OutSPS.cpb_size_scale);
-				
-				for (uint32 SchedSelIdx = 0; SchedSelIdx <= OutSPS.cpb_cnt_minus1; SchedSelIdx++ )
-				{
-					Bitstream.Read(
-						OutSPS.bit_rate_value_minus1[ SchedSelIdx ],
-						OutSPS.cpb_size_value_minus1[ SchedSelIdx ],
-						OutSPS.cbr_flag[ SchedSelIdx ]);
-				}
-
-				Bitstream.Read(
-					OutSPS.initial_cpb_removal_delay_length_minus1,
-					OutSPS.cpb_removal_delay_length_minus1,
-					OutSPS.dpb_output_delay_length_minus1,
-					OutSPS.time_offset_length);
-			};
-
-			Bitstream.Read(OutSPS.nal_hrd_parameters_present_flag);
-			if (OutSPS.nal_hrd_parameters_present_flag)
-			{
-				hrd_parameters();
-			}
-			
-			Bitstream.Read(OutSPS.vcl_hrd_parameters_present_flag);
-			if (OutSPS.vcl_hrd_parameters_present_flag)
-			{
-				hrd_parameters();
-			}
-
-			if (OutSPS.nal_hrd_parameters_present_flag || OutSPS.vcl_hrd_parameters_present_flag)
-			{
-				Bitstream.Read(OutSPS.low_delay_hrd_flag);
-			}
+		Bitstream.Read(OutSPS.vui_parameters_present_flag); // u(1)
 
-			Bitstream.Read(OutSPS.pic_struct_present_flag);
-			Bitstream.Read(OutSPS.bitstream_restriction_flag);
-			if( OutSPS.bitstream_restriction_flag )
-			{
-				Bitstream.Read(
-					OutSPS.motion_vectors_over_pic_boundaries_flag,
-					OutSPS.max_bytes_per_pic_denom,
-					OutSPS.max_bits_per_mb_denom,
-					OutSPS.log2_max_mv_length_horizontal,
-					OutSPS.log2_max_mv_length_vertical,
-					OutSPS.max_num_reorder_frames,
-					OutSPS.max_dec_frame_buffering);
-			}
-		}
+        // TODO (william.belcher): Parsing of VUI info
 
 		return EAVResult::Success;
 	}
 
-	FAVResult ParsePPS(FBitstreamReader& Bitstream, FNaluInfo const& InNaluInfo, TMap<uint32, SPS_t> const& InMapSPS, TMap<uint32, PPS_t>& OutMapPPS)
+	FAVResult ParsePPS(FBitstreamReader& Bitstream, FNaluH264 const& InNaluInfo, TMap<uint32, SPS_t> const& InMapSPS, TMap<uint32, PPS_t>& OutMapPPS)
 	{	
 		FNalu::UE pic_parameter_set_id;
-		Bitstream.Read(pic_parameter_set_id);
+		Bitstream.Read(pic_parameter_set_id); // ue(v)
 
 		PPS_t& OutPPS = OutMapPPS.FindOrAdd(pic_parameter_set_id);
 		OutPPS.pic_parameter_set_id = pic_parameter_set_id;
 		
 		Bitstream.Read(
-			OutPPS.seq_parameter_set_id,
-			OutPPS.entropy_coding_mode_flag,
-			OutPPS.bottom_field_pic_order_in_frame_present_flag,
-			OutPPS.num_slice_groups_minus1);
+			OutPPS.seq_parameter_set_id,                         // ue(v)
+			OutPPS.entropy_coding_mode_flag,                     // u(1)
+			OutPPS.bottom_field_pic_order_in_frame_present_flag, // u(1)
+			OutPPS.num_slice_groups_minus1);                     // ue(v)
 		
 		if( OutPPS.num_slice_groups_minus1 > 0 )
 		{
-			Bitstream.Read(OutPPS.slice_group_map_type);
+			Bitstream.Read(OutPPS.slice_group_map_type); // ue(v)
 			if( OutPPS.slice_group_map_type == 0 )
 			{
 				OutPPS.run_length_minus1.AddUninitialized(OutPPS.num_slice_groups_minus1 + 1);
 				for (uint32 iGroup = 0; iGroup <= OutPPS.num_slice_groups_minus1; iGroup++ )
 				{
-					Bitstream.Read(OutPPS.run_length_minus1[ iGroup ]);
+					Bitstream.Read(OutPPS.run_length_minus1[ iGroup ]); // ue(v)
 				}
 			}
+            else if( OutPPS.slice_group_map_type == 1 )
+            {
+                // TODO (wiliam.belcher): Implement support for dispersed slice group map type.
+                // See 8.2.2.2 Specification for dispersed slice group map type.
+            }
 			else if( OutPPS.slice_group_map_type == 2 )
 			{
 				OutPPS.top_left.AddUninitialized(OutPPS.num_slice_groups_minus1 + 1);
@@ -417,8 +329,8 @@ namespace UE::AVCodecCore::H264
 				for (uint32 iGroup = 0; iGroup < OutPPS.num_slice_groups_minus1; iGroup++ )
 					{
 						Bitstream.Read(
-							OutPPS.top_left[ iGroup ],
-							OutPPS.bottom_right[ iGroup ]);
+							OutPPS.top_left[ iGroup ],      // ue(v)
+							OutPPS.bottom_right[ iGroup ]); // ue(v)
 					}
 			}
 			else if( OutPPS.slice_group_map_type == 3 ||
@@ -426,12 +338,12 @@ namespace UE::AVCodecCore::H264
 					 OutPPS.slice_group_map_type == 5 )
 			{
 				Bitstream.Read(
-					OutPPS.slice_group_change_direction_flag,
-					OutPPS.slice_group_change_rate_minus1);
+					OutPPS.slice_group_change_direction_flag, // u(1)
+					OutPPS.slice_group_change_rate_minus1);   // ue(v)
 			}
 			else if( OutPPS.slice_group_map_type == 6 )
 			{
-				Bitstream.Read(OutPPS.pic_size_in_map_units_minus1);
+				Bitstream.Read(OutPPS.pic_size_in_map_units_minus1); // ue(v)
 				unsigned long numBits = 0;
 			#if defined(_MSC_VER)
 					_BitScanForward(&numBits, OutPPS.pic_size_in_map_units_minus1 + 1);
@@ -440,21 +352,21 @@ namespace UE::AVCodecCore::H264
 			#endif
 				for( uint32 i = 0; i <= OutPPS.pic_size_in_map_units_minus1; i++ )
 				{
-					Bitstream.ReadBits(OutPPS.slice_group_id[i], numBits);
+					Bitstream.ReadBits(OutPPS.slice_group_id[i], numBits); // u(numBits)
 				}
 			}
 		}
 		
-		Bitstream.Read(OutPPS.num_ref_idx_l0_default_active_minus1,
-			OutPPS.num_ref_idx_l1_default_active_minus1,
-			OutPPS.weighted_pred_flag,
-			OutPPS.weighted_bipred_idc,
-			OutPPS.pic_init_qp_minus26,
-			OutPPS.pic_init_qs_minus26,
-			OutPPS.chroma_qp_index_offset,
-			OutPPS.deblocking_filter_control_present_flag,
-			OutPPS.constrained_intra_pred_flag,
-			OutPPS.redundant_pic_cnt_present_flag);
+		Bitstream.Read(OutPPS.num_ref_idx_l0_default_active_minus1, // ue(v)
+			OutPPS.num_ref_idx_l1_default_active_minus1,            // ue(v)
+			OutPPS.weighted_pred_flag,                              // u(1)
+			OutPPS.weighted_bipred_idc,                             // u(2)
+			OutPPS.pic_init_qp_minus26,                             // se(v)
+			OutPPS.pic_init_qs_minus26,                             // se(v)
+			OutPPS.chroma_qp_index_offset,                          // se(v)
+			OutPPS.deblocking_filter_control_present_flag,          // u(1)
+			OutPPS.constrained_intra_pred_flag,                     // u(1)
+			OutPPS.redundant_pic_cnt_present_flag);                 // u(1)
 
 		auto more_rbsp_data = [&Bitstream]() -> bool
 		{
@@ -466,76 +378,64 @@ namespace UE::AVCodecCore::H264
 			return false;
 		};
 
-		if( more_rbsp_data() )
-		{
-			Bitstream.Read(OutPPS.transform_8x8_mode_flag);
-			
-			Bitstream.Read(OutPPS.pic_scaling_matrix_present_flag);
-			if (OutPPS.pic_scaling_matrix_present_flag)
-			{
-				ParseScalingList(Bitstream, InMapSPS[OutPPS.seq_parameter_set_id].chroma_format_idc, OutPPS.pic_scaling_list_present_flag, OutPPS.ScalingList4x4, OutPPS.ScalingList8x8);
-			}
-
-			Bitstream.Read(OutPPS.second_chroma_qp_index_offset);
-		}
+        // TODO (william.belcher): Parsing of additional rbsp data
 		
 		return EAVResult::Success;
 	}
 
-	FAVResult ParseSliceHeader(FBitstreamReader& Bitstream, FNaluInfo const& InNaluInfo, TMap<uint32, SPS_t> InMapSPS, TMap<uint32, PPS_t> const& InMapPPS, Slice_t& OutSlice)
+	FAVResult ParseSliceHeader(FBitstreamReader& Bitstream, FNaluH264 const& InNaluInfo, TMap<uint32, SPS_t> const& InMapSPS, TMap<uint32, PPS_t> const& InMapPPS, Slice_t& OutSlice)
 	{
-		Bitstream.Read(
-			OutSlice.first_mb_in_slice,
-			OutSlice.slice_type,
-			OutSlice.pic_parameter_set_id);
-
+        Bitstream.Read(OutSlice.first_mb_in_slice,     // ue(v)
+                       OutSlice.slice_type,            // ue(v)
+                       OutSlice.pic_parameter_set_id); // ue(v)
+        
 		const PPS_t& CurrentPPS = InMapPPS[OutSlice.pic_parameter_set_id];
 		const SPS_t& CurrentSPS = InMapSPS[CurrentPPS.seq_parameter_set_id];
 
 		if (CurrentSPS.separate_colour_plane_flag)
 		{
-			Bitstream.Read(OutSlice.colour_plane_id);
+			Bitstream.Read(OutSlice.colour_plane_id); // u(2)
 		}
 
-		Bitstream.ReadBits(OutSlice.frame_num, CurrentSPS.log2_max_frame_num_minus4 + 4);
+		Bitstream.ReadBits(OutSlice.frame_num, CurrentSPS.log2_max_frame_num_minus4 + 4); // u(v)
 
 		if (!CurrentSPS.frame_mbs_only_flag)
 		{
-			Bitstream.Read(OutSlice.field_pic_flag);
+			Bitstream.Read(OutSlice.field_pic_flag); // u(1)
 			if (OutSlice.field_pic_flag)
 			{
-				Bitstream.Read(OutSlice.bottom_field_flag);
+				Bitstream.Read(OutSlice.bottom_field_flag); // u(1)
 			}
 		}
 
 		if (InNaluInfo.Type == ENaluType::SliceIdrPicture)
 		{
-			Bitstream.Read(OutSlice.idr_pic_id);
+			Bitstream.Read(OutSlice.idr_pic_id); // ue(v)
 		}
 
 		if (CurrentSPS.pic_order_cnt_type == 0)
 		{
-			Bitstream.ReadBits(OutSlice.pic_order_cnt_lsb, CurrentSPS.log2_max_pic_order_cnt_lsb_minus4 + 4);
+			Bitstream.ReadBits(OutSlice.pic_order_cnt_lsb, CurrentSPS.log2_max_pic_order_cnt_lsb_minus4 + 4); // u(v)
 
 			if (CurrentPPS.bottom_field_pic_order_in_frame_present_flag && !OutSlice.field_pic_flag)
 			{
-				Bitstream.Read(OutSlice.delta_pic_order_cnt_bottom);
+				Bitstream.Read(OutSlice.delta_pic_order_cnt_bottom); // se(v)
 			}
 		}
 
 		if (CurrentSPS.pic_order_cnt_type == 1 && !CurrentSPS.delta_pic_order_always_zero_flag)
 		{
-			Bitstream.Read(OutSlice.delta_pic_order_cnt[0]);
+			Bitstream.Read(OutSlice.delta_pic_order_cnt[0]); // se(v)
 
 			if (CurrentPPS.bottom_field_pic_order_in_frame_present_flag && !OutSlice.field_pic_flag)
 			{
-				Bitstream.Read(OutSlice.delta_pic_order_cnt[1]);
+				Bitstream.Read(OutSlice.delta_pic_order_cnt[1]); // se(v)
 			}
 		}
 
 		if (CurrentPPS.redundant_pic_cnt_present_flag)
 		{
-			Bitstream.Read(OutSlice.redundant_pic_cnt);
+			Bitstream.Read(OutSlice.redundant_pic_cnt); // ue(v)
 		}
 
 		const bool IsP = OutSlice.slice_type == 0 || OutSlice.slice_type == 5;
@@ -546,20 +446,20 @@ namespace UE::AVCodecCore::H264
 		
 		if (IsB)
 		{
-			Bitstream.Read(OutSlice.direct_spatial_mv_pred_flag);
+			Bitstream.Read(OutSlice.direct_spatial_mv_pred_flag); // u(1)
 		}
 
 		if ( IsP || IsSP || IsB )
 		{
-			Bitstream.Read(OutSlice.num_ref_idx_active_override_flag);
+			Bitstream.Read(OutSlice.num_ref_idx_active_override_flag); // u(1)
 
 			if (OutSlice.num_ref_idx_active_override_flag)
 			{
-				Bitstream.Read(OutSlice.num_ref_idx_l0_active_minus1);
+				Bitstream.Read(OutSlice.num_ref_idx_l0_active_minus1); // ue(v)
 
 				if ( IsB )
 				{
-					Bitstream.Read(OutSlice.num_ref_idx_l1_active_minus1);
+					Bitstream.Read(OutSlice.num_ref_idx_l1_active_minus1); // ue(v)
 				}
 			}
 		}
@@ -573,26 +473,28 @@ namespace UE::AVCodecCore::H264
 		{
 			if ( OutSlice.slice_type % 5 != 2 &&  OutSlice.slice_type % 5 != 4 )
 			{
-				Bitstream.Read(OutSlice.ref_pic_list_modification_flag_l0);
+				Bitstream.Read(OutSlice.ref_pic_list_modification_flag_l0); // u(1)
 				if (OutSlice.ref_pic_list_modification_flag_l0)
 				{
-					uint32 modification_of_pic_nums_idc;
+					FNalu::UE modification_of_pic_nums_idc;
 					do
 					{
-						Bitstream.Read(modification_of_pic_nums_idc);
+						Bitstream.Read(modification_of_pic_nums_idc); // ue(v)
 						if (modification_of_pic_nums_idc == 0 || modification_of_pic_nums_idc == 1)
 						{
 							FNalu::UE pic_num;
-							Bitstream.Read(pic_num);
+							Bitstream.Read(pic_num); // ue(v)
 							
-							OutSlice.RefPicList0.Add({false, pic_num});	
+                            // TODO (william.belcher)
+							// OutSlice.RefPicList0.Add({false, pic_num});	
 						}
 						else if (modification_of_pic_nums_idc == 2)
 						{
 							FNalu::UE pic_num;
-							Bitstream.Read(pic_num);
+							Bitstream.Read(pic_num); // ue(v)
 							
-							OutSlice.RefPicList0.Add({true, pic_num});							
+                            // TODO (william.belcher)
+							// OutSlice.RefPicList0.Add({true, pic_num});							
 						}
 					}
 					while (modification_of_pic_nums_idc != 3);
@@ -601,26 +503,28 @@ namespace UE::AVCodecCore::H264
 
 			if (OutSlice.slice_type % 5 == 1)
 			{
-				Bitstream.Read(OutSlice.ref_pic_list_modification_flag_l1);
+				Bitstream.Read(OutSlice.ref_pic_list_modification_flag_l1); // u(1)
 				if (OutSlice.ref_pic_list_modification_flag_l1)
 				{
-					uint32 modification_of_pic_nums_idc;
+					FNalu::UE modification_of_pic_nums_idc;
 					do
 					{
-						Bitstream.Read(modification_of_pic_nums_idc);
+						Bitstream.Read(modification_of_pic_nums_idc); // ue(v)
 						if (modification_of_pic_nums_idc == 0 || modification_of_pic_nums_idc == 1)
 						{
 							FNalu::UE pic_num;
-							Bitstream.Read(pic_num);
+							Bitstream.Read(pic_num); // ue(v)
 							
-							OutSlice.RefPicList1.Add({false, pic_num});	
+                            // TODO (william.belcher)
+							// OutSlice.RefPicList1.Add({false, pic_num});	
 						}
 						else if (modification_of_pic_nums_idc == 2)
 						{
 							FNalu::UE pic_num;
-							Bitstream.Read(pic_num);
+							Bitstream.Read(pic_num); // ue(v)
 							
-							OutSlice.RefPicList1.Add({true, pic_num});							
+                            // TODO (william.belcher)
+							// OutSlice.RefPicList1.Add({true, pic_num});							
 						}
 					}
 					while (modification_of_pic_nums_idc != 3);
@@ -636,33 +540,67 @@ namespace UE::AVCodecCore::H264
 		if (InNaluInfo.RefIdc != 0)
 		{
 			// TODO dec_ref_pic_marking()
+            if(InNaluInfo.Type == ENaluType::SliceIdrPicture)
+            {
+                Bitstream.Read(OutSlice.no_output_of_prior_pic_flag, // u(1)
+                               OutSlice.long_term_reference_flag);   // u(1)
+            }
+            else
+            {
+                Bitstream.Read(OutSlice.adaptive_ref_pic_marking_mode_flag); // u(1)
+                if(OutSlice.adaptive_ref_pic_marking_mode_flag)
+                {
+                    FNalu::UE memory_management_control_operation;
+                    do
+                    {
+                        Bitstream.Read(memory_management_control_operation); // ue(v)
+                        if(memory_management_control_operation == 1 || memory_management_control_operation == 3)
+                        {
+                            Bitstream.Read(OutSlice.difference_of_pic_nums_minus1); // ue(v)
+                        }
+                        if(memory_management_control_operation == 2)
+                        {
+                            Bitstream.Read(OutSlice.long_term_pic_num); // ue(v)
+                        }
+                        if (memory_management_control_operation == 3 || memory_management_control_operation == 6) 
+                        {
+                            Bitstream.Read(OutSlice.long_term_frame_idx); // ue(v)
+                        }
+                        if (memory_management_control_operation == 4) 
+                        {
+                            Bitstream.Read(OutSlice.max_long_term_frame_idx_plus1); // ue(v)
+                        }
+                    } 
+                    while(memory_management_control_operation != 0);
+                }
+            }
 		}
 
 		if (CurrentPPS.entropy_coding_mode_flag && !IsI && !IsSI)
 		{
-			Bitstream.Read(OutSlice.cabac_init_idc);
+			Bitstream.Read(OutSlice.cabac_init_idc); // ue(v)
 		}
 
-		Bitstream.Read(OutSlice.slice_qp_delta);
+		Bitstream.Read(OutSlice.slice_qp_delta); // se(v)
 
 		if ( IsSP || IsSI )
 		{
 			if (IsSP)
 			{
-				Bitstream.Read(OutSlice.sp_for_switch_flag);
+				Bitstream.Read(OutSlice.sp_for_switch_flag); // u(1)
 			}
 
-			Bitstream.Read(OutSlice.slice_qs_delta);
+			Bitstream.Read(OutSlice.slice_qs_delta); // se(v)
 		}
 
 		if (CurrentPPS.deblocking_filter_control_present_flag)
 		{
-			Bitstream.Read(OutSlice.disable_deblocking_filter_idc);
+			Bitstream.Read(OutSlice.disable_deblocking_filter_idc); // ue(v)
 
 			if (OutSlice.disable_deblocking_filter_idc != 1)
 			{
-				Bitstream.Read(OutSlice.slice_alpha_c0_offset_div2);
-				Bitstream.Read(OutSlice.slice_beta_offset_div2);
+				Bitstream.Read(OutSlice.slice_alpha_c0_offset_div2); // se(v)
+				Bitstream.Read(OutSlice.slice_beta_offset_div2);     // se(v)
 			}
 		}
 
@@ -671,134 +609,10 @@ namespace UE::AVCodecCore::H264
 			&& CurrentPPS.slice_group_map_type <= 5)
 		{
 			const uint32 NumBits = FMath::CeilLogTwo((CurrentPPS.pic_size_in_map_units_minus1 + 1) / ( CurrentPPS.slice_group_change_rate_minus1 + 2) );
-			Bitstream.ReadBits(OutSlice.slice_group_change_cycle, NumBits);
+			Bitstream.ReadBits(OutSlice.slice_group_change_cycle, NumBits); // u(v)
 		}
 		
 		return EAVResult::Success;
 	}
 
-} // namespace UE::AVCodecCore::H264
-
-// template <>
-// DLLEXPORT FAVResult FAVExtension::TransformPacket(FVideoPacketH264& OutPacket, FVideoPacket const& InPacket)
-// {
-// 	using namespace UE::AVCodecCore::H264;
-
-// 	TArray<FNaluInfo> FoundNalus;
-
-// 	FAVResult Result;
-// 	Result = FindNALUs(InPacket, FoundNalus);
-
-// 	if (Result.IsNotSuccess())
-// 	{
-// 		return Result;
-// 	}
-
-// 	for (auto& NaluInfo : FoundNalus)
-// 	{		
-// 		Result = EAVResult::Success;
-		
-// 		// NALUs are is usually an EBSP so we need to strip out the emulation prevention 3 byte making them RBSP
-// 		TArray64<uint8> RBSP;
-// 		RBSP.AddUninitialized(NaluInfo.Size);
-
-// 		int32 RBSPsize = EBSPtoRBSP(RBSP.GetData(), NaluInfo.Data, NaluInfo.Size);
-// 		FBitDataStream Bitstream(RBSP.GetData(), RBSPsize);
-
-// 		switch (NaluInfo.Type)
-// 		{
-// 			case ENaluType::Unspecified:
-// 			default:
-// 				Result = FAVResult(EAVResult::Error, TEXT("Recieved ENaluType::Unspecified"), TEXT("BitstreamParserH264"));
-// 				break;
-// 			case ENaluType::SliceOfNonIdrPicture:
-// 				OutPacket.Slices.AddUninitialized();
-// 				Result = ParseSliceHeader(Bitstream, NaluInfo, OutPacket.SPS, OutPacket.PPS, OutPacket.Slices.Last());
-// 				break;
-// 			case ENaluType::SliceDataPartitionA:
-// 				break;
-// 			case ENaluType::SliceDataPartitionB:
-// 				break;
-// 			case ENaluType::SliceDataPartitionC:
-// 				break;
-// 			case ENaluType::SliceIdrPicture:
-// 				break;
-// 			case ENaluType::SupplementalEnhancementInformation:
-// 				OutPacket.SEI.AddUninitialized();
-// 				Result = ParseSEI(Bitstream, NaluInfo, OutPacket.SEI.Last());
-// 				break;
-// 			case ENaluType::SequenceParameterSet:
-// 				Result = ParseSPS(Bitstream, NaluInfo, OutPacket.SPS);
-// 				break;
-// 			case ENaluType::PictureParameterSet:
-// 				Result = ParsePPS(Bitstream, NaluInfo, OutPacket.SPS, OutPacket.PPS);
-// 				break;
-// 			case ENaluType::AccessUnitDelimiter:
-// 				break;
-// 			case ENaluType::EndOfSequence:
-// 				break;
-// 			case ENaluType::EndOfStream:
-// 				break;
-// 			case ENaluType::FillerData:
-// 				break;
-// 			case ENaluType::SequenceParameterSetExtension:
-// 				break;
-// 			case ENaluType::PrefixNalUnit:
-// 				break;
-// 			case ENaluType::SubsetSequenceParameterSet:
-// 				break;
-// 			case ENaluType::Reserved16:
-// 				Result = FAVResult(EAVResult::Error, TEXT("Recieved ENaluType::Reserved16"), TEXT("BitstreamParserH264"));
-// 				break;
-// 			case ENaluType::Reserved17:
-// 				Result = FAVResult(EAVResult::Error, TEXT("Recieved ENaluType::Reserved17"), TEXT("BitstreamParserH264"));
-// 				break;
-// 			case ENaluType::Reserved18:
-// 				Result = FAVResult(EAVResult::Error, TEXT("Recieved ENaluType::Reserved18"), TEXT("BitstreamParserH264"));
-// 				break;
-// 			case ENaluType::SliceOfAnAuxiliaryCoded:
-// 				break;
-// 			case ENaluType::SliceExtension:
-// 				break;
-// 			case ENaluType::SliceExtensionForDepthView:
-// 				break;
-// 			case ENaluType::Reserved22:
-// 				Result = FAVResult(EAVResult::Error, TEXT("Recieved ENaluType::Reserved22"), TEXT("BitstreamParserH264"));
-// 				break;
-// 			case ENaluType::Reserved23:
-// 				Result = FAVResult(EAVResult::Error, TEXT("Recieved ENaluType::Reserved23"), TEXT("BitstreamParserH264"));
-// 				break;
-// 			case ENaluType::Unspecified24:
-// 				Result = FAVResult(EAVResult::Error, TEXT("Recieved ENaluType::Unspecified24"), TEXT("BitstreamParserH264"));
-// 				break;
-// 			case ENaluType::Unspecified25:
-// 				Result = FAVResult(EAVResult::Error, TEXT("Recieved ENaluType::Unspecified25"), TEXT("BitstreamParserH264"));
-// 				break;
-// 			case ENaluType::Unspecified26:
-// 				Result = FAVResult(EAVResult::Error, TEXT("Recieved ENaluType::Unspecified26"), TEXT("BitstreamParserH264"));
-// 				break;
-// 			case ENaluType::Unspecified27:
-// 				Result = FAVResult(EAVResult::Error, TEXT("Recieved ENaluType::Unspecified27"), TEXT("BitstreamParserH264"));
-// 				break;
-// 			case ENaluType::Unspecified28:
-// 				Result = FAVResult(EAVResult::Error, TEXT("Recieved ENaluType::Unspecified28"), TEXT("BitstreamParserH264"));
-// 				break;
-// 			case ENaluType::Unspecified29:
-// 				Result = FAVResult(EAVResult::Error, TEXT("Recieved ENaluType::Unspecified29"), TEXT("BitstreamParserH264"));
-// 				break;
-// 			case ENaluType::Unspecified30:
-// 				Result = FAVResult(EAVResult::Error, TEXT("Recieved ENaluType::Unspecified30"), TEXT("BitstreamParserH264"));
-// 				break;
-// 			case ENaluType::Unspecified31:
-// 				Result = FAVResult(EAVResult::Error, TEXT("Recieved ENaluType::Unspecified31"), TEXT("BitstreamParserH264"));
-// 				break;
-// 		}
-
-// 		if (Result.IsNotSuccess())
-// 		{
-// 			return Result;
-// 		}
-// 	}
-
-// 	return Result;
-// }
+} // namespace UE::AVCodecCore::H264
\ No newline at end of file
