Add: VideoToolbox Encoder and Decoders to AVCodecs

diff --git a/Engine/Plugins/Experimental/AVCodecs/AVCodecsCore/Source/AVCodecsCore/Private/Video/Decoders/Configs/VideoDecoderConfigH264.cpp b/Engine/Plugins/Experimental/AVCodecs/AVCodecsCore/Source/AVCodecsCore/Private/Video/Decoders/Configs/VideoDecoderConfigH264.cpp
--- a/Engine/Plugins/Experimental/AVCodecs/AVCodecsCore/Source/AVCodecsCore/Private/Video/Decoders/Configs/VideoDecoderConfigH264.cpp
+++ b/Engine/Plugins/Experimental/AVCodecs/AVCodecsCore/Source/AVCodecsCore/Private/Video/Decoders/Configs/VideoDecoderConfigH264.cpp
@@ -6,11 +6,11 @@
 
 REGISTER_TYPEID(FVideoDecoderConfigH264);
 
-FAVResult FVideoDecoderConfigH264::Parse(TSharedRef<FAVInstance> const& Instance, FVideoPacket const& Packet, TArray<UE::AVCodecCore::H264::Slice_t>& OutSlices)
+FAVResult FVideoDecoderConfigH264::Parse(FVideoPacket const& Packet, TArray<UE::AVCodecCore::H264::Slice_t>& OutSlices)
 {
 	using namespace UE::AVCodecCore::H264;
 
-	TArray<FNaluInfo> FoundNalus;
+	TArray<FNaluH264> FoundNalus;
 
 	FAVResult Result;
 	Result = FindNALUs(Packet, FoundNalus);
@@ -136,3 +136,19 @@ FAVResult FVideoDecoderConfigH264::Parse(TSharedRef<FAVInstance> const& Instance
 
 	return Result;
 }
+
+TOptional<int> FVideoDecoderConfigH264::GetLastSliceQP(TArray<UE::AVCodecCore::H264::Slice_t>& Slices)
+{
+    using namespace UE::AVCodecCore::H264;
+
+    Slice_t& LastSlice = Slices.Last();
+    PPS_t& LastSlicePPS = PPS[LastSlice.pic_parameter_set_id];
+
+    const int QP = 26 + LastSlicePPS.pic_init_qp_minus26 + LastSlice.slice_qp_delta;
+    if(QP < 0 || QP > 51)
+    {
+		FAVResult::Log(EAVResult::Warning, FString::Printf(TEXT("QP of %d was outside of range 0 to 51 inclusive. Defaulting to 0"), QP), TEXT("H264"));
+        return TOptional<int>();
+    }
+    return QP;
+}
diff --git a/Engine/Plugins/Experimental/AVCodecs/AVCodecsCore/Source/AVCodecsCore/Private/Video/Decoders/Configs/VideoDecoderConfigH265.cpp b/Engine/Plugins/Experimental/AVCodecs/AVCodecsCore/Source/AVCodecsCore/Private/Video/Decoders/Configs/VideoDecoderConfigH265.cpp
--- a/Engine/Plugins/Experimental/AVCodecs/AVCodecsCore/Source/AVCodecsCore/Private/Video/Decoders/Configs/VideoDecoderConfigH265.cpp
+++ b/Engine/Plugins/Experimental/AVCodecs/AVCodecsCore/Source/AVCodecsCore/Private/Video/Decoders/Configs/VideoDecoderConfigH265.cpp
@@ -7,7 +7,7 @@
 
 REGISTER_TYPEID(FVideoDecoderConfigH265);
 
-FAVResult FVideoDecoderConfigH265::Parse(TSharedRef<FAVInstance> const& Instance, FVideoPacket const& Packet, TArray<TSharedPtr<UE::AVCodecCore::H265::FNaluH265>>& Nalus)
+FAVResult FVideoDecoderConfigH265::Parse(FVideoPacket const& Packet, TArray<UE::AVCodecCore::H265::Slice_t>& OutSlices)
 {
 	using namespace UE::AVCodecCore::H265;
 
@@ -21,15 +21,17 @@ FAVResult FVideoDecoderConfigH265::Parse(TSharedRef<FAVInstance> const& Instance
 		return Result;
 	}
 
-	FVideoDecoderConfigH265& H265 = Instance->Edit<FVideoDecoderConfigH265>();
-	TSharedPtr<FNaluVPS> CandidateVPS;
-	TSharedPtr<FNaluSPS> CandidateSPS;
-	TSharedPtr<FNaluPPS> CandidatePPS;
-
-	for (auto& Nalu : FoundNalus)
+	for (auto& NaluInfo : FoundNalus)
 	{
 		Result = EAVResult::Success;
-		switch (Nalu.nal_unit_type)
+
+        TArray64<uint8> RBSP;
+		RBSP.AddUninitialized(NaluInfo.Size);
+
+		int32 RBSPsize = EBSPtoRBSP(RBSP.GetData(), NaluInfo.Data, NaluInfo.Size);
+		FBitstreamReader BitStream(RBSP.GetData(), RBSPsize);
+
+		switch (NaluInfo.Type)
 		{
 			case ENaluType::TRAIL_N:
 			case ENaluType::TRAIL_R:
@@ -41,9 +43,8 @@ FAVResult FVideoDecoderConfigH265::Parse(TSharedRef<FAVInstance> const& Instance
 			case ENaluType::RADL_R:
 			case ENaluType::RASL_N:
 			case ENaluType::RASL_R:
-				Nalus.Add(MakeShared<FNaluSlice>(Nalu));
-				StaticCastSharedPtr<FNaluSlice>(Nalus.Last())->Parse(this);
-				UE_LOG(LogTemp, Verbose, TEXT("H256 Parsing: Recieved Slice"));
+                OutSlices.AddUninitialized();
+                Result = ParseSliceHeader(BitStream, NaluInfo, VPS, SPS, PPS, OutSlices.Last());
 				break;
 			case ENaluType::RSV_VCL_N10:
 				checkNoEntry(); // Reserved
@@ -70,9 +71,8 @@ FAVResult FVideoDecoderConfigH265::Parse(TSharedRef<FAVInstance> const& Instance
 			case ENaluType::BLA_W_RADL:
 			case ENaluType::BLA_N_LP:
 			case ENaluType::CRA_NUT:
-				Nalus.Add(MakeShared<FNaluSlice>(Nalu));
-				StaticCastSharedPtr<FNaluSlice>(Nalus.Last())->Parse(this);
-				UE_LOG(LogTemp, Verbose, TEXT("H256 Parsing: Recieved Slice"));
+				OutSlices.AddUninitialized();
+                Result = ParseSliceHeader(BitStream, NaluInfo, VPS, SPS, PPS, OutSlices.Last());
 				break;
 			case ENaluType::RSV_IRAP_VCL22:
 				checkNoEntry(); // Reserved or Unspecified
@@ -105,47 +105,25 @@ FAVResult FVideoDecoderConfigH265::Parse(TSharedRef<FAVInstance> const& Instance
 				checkNoEntry(); // Reserved
 				break;
 			case ENaluType::VPS_NUT:
-				CandidateVPS = MakeShared<FNaluVPS>(Nalu);
-				CandidateVPS->Parse();
-				H265.ParsedVPS.Add(CandidateVPS->vps_video_parameter_set_id, CandidateVPS);
-				Nalus.Add(CandidateVPS);
-				UE_LOG(LogTemp, Verbose, TEXT("H256 Parsing: Recieved VPS %u"), CandidateVPS->vps_video_parameter_set_id.Value);
+                Result = ParseVPS(BitStream, NaluInfo, VPS);
 				break;
 			case ENaluType::SPS_NUT:
-				CandidateSPS = MakeShared<FNaluSPS>(Nalu);
-				CandidateSPS->Parse();
-				H265.ParsedSPS.Add(CandidateSPS->sps_video_parameter_set_id, CandidateSPS);
-				Nalus.Add(CandidateSPS);
-				UE_LOG(LogTemp, Verbose, TEXT("H256 Parsing: Recieved SPS %u"), CandidateSPS->sps_video_parameter_set_id.Value);
+				Result = ParseSPS(BitStream, NaluInfo, VPS, SPS);
 				break;
 			case ENaluType::PPS_NUT:
-				CandidatePPS = MakeShared<FNaluPPS>(Nalu);
-				CandidatePPS->Parse();
-				H265.ParsedPPS.Add(CandidatePPS->pps_pic_parameter_set_id, CandidatePPS);
-				Nalus.Add(CandidatePPS);
-				UE_LOG(LogTemp, Verbose, TEXT("H256 Parsing: Recieved PPS %u"), CandidatePPS->pps_pic_parameter_set_id.Value);
-				break;
+				Result = ParsePPS(BitStream, NaluInfo, VPS, SPS, PPS);
 			case ENaluType::AUD_NUT:
-				Nalus.Add(MakeShared<FNaluH265>(Nalu));
-				UE_LOG(LogTemp, Verbose, TEXT("H256 Parsing: Recieved Access Unit Delimiter"));
 				break;
 			case ENaluType::EOS_NUT:
-				Nalus.Add(MakeShared<FNaluH265>(Nalu));
-				UE_LOG(LogTemp, Verbose, TEXT("H256 Parsing: Recieved End Of Sequence"));
 				break;
 			case ENaluType::EOB_NUT:
-				Nalus.Add(MakeShared<FNaluH265>(Nalu));
-				UE_LOG(LogTemp, Verbose, TEXT("H256 Parsing: Recieved End Of Bitstream"));
 				break;
 			case ENaluType::FD_NUT:
-				Nalus.Add(MakeShared<FNaluH265>(Nalu));
-				UE_LOG(LogTemp, Verbose, TEXT("H256 Parsing: Recieved Filler Data"));
 				break;
 			case ENaluType::PREFIX_SEI_NUT:
 			case ENaluType::SUFFIX_SEI_NUT:
-				H265.ParsedSEI.Add(MakeShared<FNaluSEI>(Nalu));
-				H265.ParsedSEI.Last()->Parse();
-				UE_LOG(LogTemp, Verbose, TEXT("H256 Parsing: Recieved SEI"));
+				SEI.AddUninitialized();
+				Result = ParseSEI(BitStream, NaluInfo, SEI.Last());
 				break;
 			case ENaluType::RSV_NVCL41:
 				checkNoEntry(); // Reserved or Unspecified
@@ -185,251 +163,18 @@ FAVResult FVideoDecoderConfigH265::Parse(TSharedRef<FAVInstance> const& Instance
 	return Result;
 }
 
-FAVResult FVideoDecoderConfigH265::UpdateRPS(TSharedRef<UE::AVCodecCore::H265::FNaluSlice> CurrentSlice)
+TOptional<int> FVideoDecoderConfigH265::GetLastSliceQP(TArray<UE::AVCodecCore::H265::Slice_t>& Slices)
 {
-	using namespace  UE::AVCodecCore::H265;
-
-	ReferencePictureSet.PicOrderCntVal[LastRpsIdx] = LastPicOrderCntValue;
-
-	// Invalidate all of the usages when a new Intra Decode Refresh frame comes in
-	if (CurrentSlice->IsIDR())
-	{
-		// Clear the current RPS working state and usage
-		ReferencePictureSet.PrepState();
-
-		LastRpsIdx = 0;
-		LastPicOrderCntValue = CurrentSlice->CurrPicOrderCntVal;
-	}
-	else
-	{		
-		TSharedPtr<FNaluSPS> PinnedSPS = CurrentSlice->CurrentSPS.Pin();
-		TSharedPtr<FNaluPPS> PinnedPPS = CurrentSlice->CurrentPPS.Pin();
-		uint8 i, j, k;
-		
-		// (8-5) Calculate NumPocTotalCur
-		uint8 const& CurrRspIdx = CurrentSlice->CurrRpsIdx;
-		short_term_ref_pic_set_t const& ShortTermRefPicSet = PinnedSPS->short_term_ref_pic_sets[CurrRspIdx]; 
-		uint32 MaxPicOrderCntLsb = 1 << (PinnedSPS->log2_max_pic_order_cnt_lsb_minus4 + 4);
-		bool UsedByCurrPicLt;
-
-		// NOTE (aidan) Should we be creating these with every slice?
-		uint8_t CurrDeltaPocMsbPresentFlag[16];
-		uint8_t FollDeltaPocMsbPresentFlag[16];
-	
-		for (i = 0, j = 0, k = 0; i < ShortTermRefPicSet.NumNegativePics; i++)
-		{
-			if (ShortTermRefPicSet.UsedByCurrPicS0[i])
-			{
-				ReferencePictureSet.PocStCurrBefore[j++] = CurrentSlice->CurrPicOrderCntVal + ShortTermRefPicSet.DeltaPocS0[i];
-			}
-			else
-			{
-				ReferencePictureSet.PocStFoll[k++] = CurrentSlice->CurrPicOrderCntVal + ShortTermRefPicSet.DeltaPocS0[i];
-			}
-		}
-		ReferencePictureSet.NumPocStCurrBefore = j;
-
-		for (i = 0, j = 0; i < ShortTermRefPicSet.NumPositivePics; i++)
-		{
-			if (ShortTermRefPicSet.UsedByCurrPicS1[i])
-			{
-				ReferencePictureSet.PocStCurrAfter[j++] = CurrentSlice->CurrPicOrderCntVal + ShortTermRefPicSet.DeltaPocS1[i];
-			}
-			else
-			{
-				ReferencePictureSet.PocStFoll[k++] = CurrentSlice->CurrPicOrderCntVal + ShortTermRefPicSet.DeltaPocS1[i];
-			}
-		}
-		ReferencePictureSet.NumPocStCurrAfter = j;
-		ReferencePictureSet.NumPocStFoll = k;
-	
-		uint8 pocLt;
-		for (i = 0, j = 0, k = 0; i < (CurrentSlice->num_long_term_sps + CurrentSlice->num_long_term_pics); i++)
-		{			
-			if (i < CurrentSlice->num_long_term_sps)
-			{
-				pocLt = PinnedSPS->long_term_ref_pics_sps[CurrentSlice->lt_idx_sps[i]].lt_ref_pic_poc_lsb_sps;
-				UsedByCurrPicLt = PinnedSPS->long_term_ref_pics_sps[CurrentSlice->lt_idx_sps[i]].used_by_curr_pic_lt_sps_flag.AsBool();
-			}
-			else
-			{
-				pocLt = CurrentSlice->poc_lsb_lt[i];
-				UsedByCurrPicLt = CurrentSlice->used_by_curr_pic_lt_flag[CurrentSlice->lt_idx_sps[i]].AsBool();
-			}
-			
-			if (CurrentSlice->delta_poc_msb_present_flag[i])
-			{
-				pocLt += CurrentSlice->CurrPicOrderCntVal - CurrentSlice->delta_poc_msb_cycle_lt[i] * MaxPicOrderCntLsb - ( CurrentSlice->CurrPicOrderCntVal & (MaxPicOrderCntLsb-1u));
-			}
-
-			if (UsedByCurrPicLt)
-			{
-				ReferencePictureSet.PocLtCurr[j] = pocLt;
-				CurrDeltaPocMsbPresentFlag[j++] = CurrentSlice->delta_poc_msb_present_flag[i];
-				ReferencePictureSet.NumPicTotalCurr++;
-			}
-			else
-			{
-				ReferencePictureSet.PocLtFoll[k] = pocLt;
-				FollDeltaPocMsbPresentFlag[k++] = CurrentSlice->delta_poc_msb_present_flag[i];
-			}
-		}
-		ReferencePictureSet.NumPocLtCurr = j;
-		ReferencePictureSet.NumPocLtFoll = k;
-
-		auto const& PicXInDPB = [&PinnedSPS, &CurrentSlice, MaxPicOrderCntLsb, PicOrderCntVal = ReferencePictureSet.PicOrderCntVal, PicLayerId = ReferencePictureSet.PicLayerId](int32 PicOrderCount, bool DeltaPocMsbPresentFlag) -> uint8
-		{
-			if (DeltaPocMsbPresentFlag)
-			{
-				for (uint8 i = 0; i < 16 && i < PinnedSPS->MaxDpbSize; i++)
-				{
-					if (PicOrderCntVal[i] == PicOrderCount && CurrentSlice->nuh_layer_id == PicLayerId[i])
-					{
-						return i;
-					}
-				}
-			}
-			else
-			{
-				for (uint8 i = 0; i < 16 && i < PinnedSPS->MaxDpbSize; i++)
-				{
-					if ((PicOrderCntVal[i] & (FMath::Max<uint32>(MaxPicOrderCntLsb, 1) - 1)) == PicOrderCount && CurrentSlice->nuh_layer_id == PicLayerId[i])
-					{
-						return i;
-					}
-				}
-			}
-		
-			return 255; // NOTE (aidan) placholder for "no reference picture"
-		};
-
-		// Clear reference usage flags ready to be set for current frame
-		FMemory::Memset(ReferencePictureSet.PicUsage, 0, sizeof(ReferencePictureSet.PicUsage));
-
-		// (8-6) Generate Long-term Reference Pictures
-		for( i = 0; i < ReferencePictureSet.NumPocLtCurr; i++ )
-		{
-			ReferencePictureSet.RefPicSetLtCurr[i] = PicXInDPB(ReferencePictureSet.PocLtCurr[i], CurrDeltaPocMsbPresentFlag[i] > 0);
-			if(ReferencePictureSet.RefPicSetLtCurr[i] != 255)
-			{
-				ReferencePictureSet.PicUsage[i] |= EPictureUsage::LONG_TERM;
-			}
-		}
-	
-		for( i = 0; i < ReferencePictureSet.NumPocLtFoll; i++ )
-		{
-			ReferencePictureSet.RefPicSetLtFoll[i] = PicXInDPB(ReferencePictureSet.PocLtFoll[i], FollDeltaPocMsbPresentFlag[i] > 0);
-			if(ReferencePictureSet.RefPicSetLtFoll[i] != 255)
-			{
-				ReferencePictureSet.PicUsage[i] |= EPictureUsage::LONG_TERM;
-			}
-		}
-
-		// (8-7) Generate short-term reference pictures
-		for( i = 0; i < ReferencePictureSet.NumPocStCurrBefore; i++ )
-		{
-			ReferencePictureSet.RefPicSetStCurrBefore[i] = PicXInDPB(ReferencePictureSet.PocStCurrBefore[i], true);
-			if(ReferencePictureSet.RefPicSetStCurrBefore[i] != 255)
-			{
-				ReferencePictureSet.PicUsage[i] |= EPictureUsage::SHORT_TERM;
-			}
-		}
-	
-		for( i = 0; i < ReferencePictureSet.NumPocStCurrAfter; i++ )
-		{
-			ReferencePictureSet.RefPicSetStCurrAfter[i] = PicXInDPB(ReferencePictureSet.PocStCurrAfter[i], true);
-			if(ReferencePictureSet.RefPicSetStCurrAfter[i] != 255)
-			{
-				ReferencePictureSet.PicUsage[i] |= EPictureUsage::SHORT_TERM;
-			}
-		}
-	
-		for( i = 0; i < ReferencePictureSet.NumPocStFoll; i++ )
-		{
-			ReferencePictureSet.RefPicSetStFoll[i] = PicXInDPB(ReferencePictureSet.PocStFoll[i], 1);
-			if(ReferencePictureSet.RefPicSetStFoll[i] != 255)
-			{
-				ReferencePictureSet.PicUsage[i] |= EPictureUsage::SHORT_TERM;
-			}
-		}
-
-		// TODO (aidan) Unavailable frames estimation. Section 8.3.3.1 talks about image generation but not sure how to kick that off
-
-		// (8-8) Store old Ref Picture lists
-		const uint32 NumRpsCurrTempList0 = FMath::Max<uint32>(CurrentSlice->num_ref_idx_l0_active_minus1.Value + 1, ReferencePictureSet.NumPicTotalCurr);
-	
-		TArray<uint32> RefPicListTemp0;
-		RefPicListTemp0.SetNumZeroed(NumRpsCurrTempList0);
-	
-		uint32 rIdx;
-		for (rIdx = 0; rIdx < NumRpsCurrTempList0; rIdx++)
-		{
-			for (i = 0; i < ReferencePictureSet.NumPocStCurrBefore && rIdx < NumRpsCurrTempList0; rIdx++, i++)
-			{
-				RefPicListTemp0[rIdx] = ReferencePictureSet.RefPicSetStCurrBefore[i];
-			}
-			for (i = 0; i < ReferencePictureSet.NumPocStCurrAfter && rIdx < NumRpsCurrTempList0; rIdx++, i++)
-			{
-				RefPicListTemp0[rIdx] = ReferencePictureSet.RefPicSetStCurrAfter[i];
-			}
-			for (i = 0; i < ReferencePictureSet.NumPocLtCurr && rIdx < NumRpsCurrTempList0; rIdx++, i++)
-			{
-				RefPicListTemp0[rIdx] = ReferencePictureSet.RefPicSetLtCurr[i];
-			}
-			if (PinnedPPS->pps_scc_extension.pps_curr_pic_ref_enabled_flag)
-			{
-				RefPicListTemp0[rIdx++] = CurrPicIdx;
-			}
-		}
-
-		// (8-9) Construct RefPicList0
-		for (rIdx = 0; rIdx <= CurrentSlice->num_ref_idx_l0_active_minus1; rIdx++)
-		{
-			ReferencePictureSet.RefPicList0[rIdx] = CurrentSlice->ref_pic_list_modification.ref_pic_list_modification_flag_l0 ? RefPicListTemp0[CurrentSlice->ref_pic_list_modification.list_entry_l0[rIdx]] : RefPicListTemp0[rIdx];
-		}
-
-		if (PinnedPPS->pps_scc_extension.pps_curr_pic_ref_enabled_flag && !CurrentSlice->ref_pic_list_modification.ref_pic_list_modification_flag_l0 && NumRpsCurrTempList0 > (CurrentSlice->num_ref_idx_l0_active_minus1 + 1))
-		{
-			ReferencePictureSet.RefPicList0[CurrentSlice->num_ref_idx_l0_active_minus1] = CurrPicIdx;
-		}
-
-		// (8-10) Repeat for RefPicList1 when decoding a B slice
-		if (CurrentSlice->slice_type == EH265SliceType::B)
-		{
-			const uint32 NumRpsCurrTempList1 = FMath::Max<uint32>(CurrentSlice->num_ref_idx_l1_active_minus1 + 1, ReferencePictureSet.NumPicTotalCurr);
-
-			rIdx = 0;
-			TArray<uint32> RefPicListTemp1;
-			RefPicListTemp1.SetNumZeroed(NumRpsCurrTempList1);
-			while (rIdx < NumRpsCurrTempList1)
-			{
-				for (i = 0; i < ReferencePictureSet.NumPocStCurrAfter && rIdx < NumRpsCurrTempList1; rIdx++, i++)
-				{
-					RefPicListTemp1[rIdx] = ReferencePictureSet.RefPicSetStCurrAfter[i];
-				}
-				for (i = 0; i < ReferencePictureSet.NumPocStCurrBefore && rIdx < NumRpsCurrTempList1; rIdx++, i++)
-				{
-					RefPicListTemp1[rIdx] = ReferencePictureSet.RefPicSetStCurrBefore[i];
-				}
-				for (i = 0; i < ReferencePictureSet.NumPocLtCurr && rIdx < NumRpsCurrTempList1; rIdx++, i++)
-				{
-					RefPicListTemp1[rIdx] = ReferencePictureSet.RefPicSetLtCurr[i];
-				}
-				if (PinnedPPS->pps_scc_extension.pps_curr_pic_ref_enabled_flag)
-				{
-					RefPicListTemp1[rIdx++] = CurrPicIdx;
-				}
-			}
-
-			for (rIdx = 0; rIdx <= CurrentSlice->num_ref_idx_l0_active_minus1; rIdx++)
-			{
-				ReferencePictureSet.RefPicList1[rIdx] = CurrentSlice->ref_pic_list_modification.ref_pic_list_modification_flag_l0 ? RefPicListTemp1[CurrentSlice->ref_pic_list_modification.list_entry_l1[rIdx]] : RefPicListTemp1[rIdx];
-			}
-		}
-
-		LastRpsIdx = CurrentSlice->CurrPicOrderCntVal % PinnedSPS->MaxDpbSize;
-		LastPicOrderCntValue = CurrentSlice->CurrPicOrderCntVal;
-	}
-
-	return EAVResult::Success;
+    using namespace UE::AVCodecCore::H265;
+
+    Slice_t& LastSlice = Slices.Last();
+    PPS_t& LastSlicePPS = PPS[LastSlice.slice_pic_parameter_set_id];
+
+    const int QP = 26 + LastSlicePPS.init_qp_minus26 + LastSlice.slice_qp_delta;
+    if(QP < 0 || QP > 51)
+    {
+		FAVResult::Log(EAVResult::Warning, FString::Printf(TEXT("QP of %d was outside of range 0 to 51 inclusive. Defaulting to 0"), QP), TEXT("H265"));
+        return TOptional<int>();
+    }
+    return QP;
 }
diff --git a/Engine/Plugins/Experimental/AVCodecs/AVCodecsCore/Source/AVCodecsCore/Private/Video/Resources/Metal/VideoResourceMetal.cpp b/Engine/Plugins/Experimental/AVCodecs/AVCodecsCore/Source/AVCodecsCore/Private/Video/Resources/Metal/VideoResourceMetal.cpp
new file mode 100644
--- /dev/null
+++ b/Engine/Plugins/Experimental/AVCodecs/AVCodecsCore/Source/AVCodecsCore/Private/Video/Resources/Metal/VideoResourceMetal.cpp
@@ -0,0 +1,57 @@
+// Copyright Epic Games, Inc. All Rights Reserved.
+
+#if AVCODECS_USE_METAL
+
+#include "Video/Resources/Metal/VideoResourceMetal.h"
+
+REGISTER_TYPEID(FVideoContextMetal);
+REGISTER_TYPEID(FVideoResourceMetal);
+
+static TAVResult<EVideoFormat> ConvertFormat(mtlpp::PixelFormat Format)
+{
+	switch (Format)
+	{
+	case mtlpp::PixelFormat::BGRA8Unorm:
+    case mtlpp::PixelFormat::BGRA8Unorm_sRGB:
+		return EVideoFormat::BGRA;
+	case mtlpp::PixelFormat::RGB10A2Unorm:
+		return EVideoFormat::ABGR10;
+	case mtlpp::PixelFormat::R8Unorm:
+	case mtlpp::PixelFormat::R8Uint:
+		return EVideoFormat::R8;	
+	default:
+		return FAVResult(EAVResult::ErrorUnsupported, FString::Printf(TEXT("mtlpp::PixelFormat format %d is not supported"), Format), TEXT("Metal"));
+	}
+}
+
+FVideoContextMetal::FVideoContextMetal(mtlpp::Device const& Device)
+	: Device(Device)
+{
+}
+
+FVideoDescriptor FVideoResourceMetal::GetDescriptorFrom(TSharedRef<FAVDevice> const& Device, mtlpp::Texture* Raw)
+{
+    uint32_t Width = Raw->GetWidth();
+    uint32_t Height = Raw->GetHeight();
+    TAVResult<EVideoFormat> ConvertedFormat = ConvertFormat(Raw->GetPixelFormat());
+    
+	return FVideoDescriptor(ConvertedFormat, Width, Height);
+}
+
+FVideoResourceMetal::FVideoResourceMetal(TSharedRef<FAVDevice> const& Device, mtlpp::Texture* Raw, FAVLayout const& Layout)
+	: TVideoResource(Device, Layout, GetDescriptorFrom(Device, Raw))
+	, Raw(Raw)
+{
+}
+
+FAVResult FVideoResourceMetal::Validate() const
+{
+	if (!Raw)
+	{
+		return FAVResult(EAVResult::ErrorInvalidState, TEXT("Raw resource is invalid"), TEXT("Metal"));
+	}
+
+	return EAVResult::Success;
+}
+
+#endif
diff --git a/Engine/Plugins/Experimental/AVCodecs/AVCodecsCore/Source/AVCodecsCore/Public/Video/CodecUtils/CodecUtilsH264.h b/Engine/Plugins/Experimental/AVCodecs/AVCodecsCore/Source/AVCodecsCore/Public/Video/CodecUtils/CodecUtilsH264.h
--- a/Engine/Plugins/Experimental/AVCodecs/AVCodecsCore/Source/AVCodecsCore/Public/Video/CodecUtils/CodecUtilsH264.h
+++ b/Engine/Plugins/Experimental/AVCodecs/AVCodecsCore/Source/AVCodecsCore/Public/Video/CodecUtils/CodecUtilsH264.h
@@ -117,7 +117,9 @@ namespace UE::AVCodecCore::H264
 		High422 = 122,
 		// High422Intra = 122,					with constraint flag 3 set
 		StereoHigh = 128,
+        MultiresolutionFrameCompatibleHigh = 134,
 		MultiviewDepthHigh = 138,
+        EnhancedMultiviewDepthHigh = 139,
 		High444 = 244,
 		// High444Intra = 244,					with constraint flag 3 set
 	};
@@ -251,7 +253,7 @@ namespace UE::AVCodecCore::H264
 		return OutBuf - OutBase;
 	}
 
-	struct FNaluInfo
+	struct FNaluH264
 	{
 		uint64 Start, Size;
 		uint8 StartCodeSize;
@@ -260,7 +262,7 @@ namespace UE::AVCodecCore::H264
 		const uint8* Data;
 	};
 
-	FAVResult FindNALUs(FVideoPacket const& InPacket, TArray<FNaluInfo>& FoundNalus);
+	AVCODECSCORE_API FAVResult FindNALUs(FVideoPacket const& InPacket, TArray<FNaluH264>& FoundNalus);
 
 	constexpr uint8 Default_4x4_Intra[16] = {
 		6, 13, 13, 20,
@@ -446,7 +448,7 @@ namespace UE::AVCodecCore::H264
 		}
 	};
 
-	FAVResult ParseSPS(FBitstreamReader& Bitstream, FNaluInfo const& InNaluInfo, TMap<uint32, SPS_t>& OutMapSPS);
+	FAVResult ParseSPS(FBitstreamReader& Bitstream, FNaluH264 const& InNaluInfo, TMap<uint32, SPS_t>& OutMapSPS);
 
 	struct PPS_t : public FNalu
 	{
@@ -481,7 +483,7 @@ namespace UE::AVCodecCore::H264
 		SE second_chroma_qp_index_offset;
 	};
 
-	FAVResult ParsePPS(FBitstreamReader& Bitstream, FNaluInfo const& InNaluInfo, TMap<uint32, SPS_t> const& InMapSPS, TMap<uint32, PPS_t>& OutMapPPS);
+	FAVResult ParsePPS(FBitstreamReader& Bitstream, FNaluH264 const& InNaluInfo, TMap<uint32, SPS_t> const& InMapSPS, TMap<uint32, PPS_t>& OutMapPPS);
 
 	struct Slice_t : public FNalu
 	{
@@ -518,9 +520,9 @@ namespace UE::AVCodecCore::H264
 		};
 
 		U<1> ref_pic_list_modification_flag_l0;
-		TArray<RefPic_t> RefPicList0;
+		// TArray<RefPic_t> RefPicList0;
 		U<1> ref_pic_list_modification_flag_l1;
-		TArray<RefPic_t> RefPicList1;
+		// TArray<RefPic_t> RefPicList1;
 
 		// pred_weight_table
 
@@ -529,14 +531,19 @@ namespace UE::AVCodecCore::H264
 		U<1> long_term_reference_flag;
 
 		U<1> adaptive_ref_pic_marking_mode_flag;
+
+        UE difference_of_pic_nums_minus1;
+        UE long_term_pic_num;
+        UE long_term_frame_idx;
+        UE max_long_term_frame_idx_plus1;
 	};
 
-	FAVResult ParseSliceHeader(FBitstreamReader& Bitstream, FNaluInfo const& InNaluInfo, TMap<uint32, SPS_t> InMapSPS, TMap<uint32, PPS_t> const& InMapPPS, Slice_t& OutSlice);
+	FAVResult ParseSliceHeader(FBitstreamReader& Bitstream, FNaluH264 const& InNaluInfo, TMap<uint32, SPS_t> const& InMapSPS, TMap<uint32, PPS_t> const& InMapPPS, Slice_t& OutSlice);
 
 	// SEI can transmit arbitrary data so we probably want something less rigid for this
 	struct SEI_t
 	{
 	};
 
-	FAVResult ParseSEI(FBitstreamReader& Bitstream, FNaluInfo const& InNaluInfo, SEI_t& OutSEI);
+	FAVResult ParseSEI(FBitstreamReader& Bitstream, FNaluH264 const& InNaluInfo, SEI_t& OutSEI);
 } // namespace UE::AVCodecCore::H264
diff --git a/Engine/Plugins/Experimental/AVCodecs/AVCodecsCore/Source/AVCodecsCore/Public/Video/CodecUtils/CodecUtilsH265.h b/Engine/Plugins/Experimental/AVCodecs/AVCodecsCore/Source/AVCodecsCore/Public/Video/CodecUtils/CodecUtilsH265.h
--- a/Engine/Plugins/Experimental/AVCodecs/AVCodecsCore/Source/AVCodecsCore/Public/Video/CodecUtils/CodecUtilsH265.h
+++ b/Engine/Plugins/Experimental/AVCodecs/AVCodecsCore/Source/AVCodecsCore/Public/Video/CodecUtils/CodecUtilsH265.h
@@ -150,11 +150,15 @@ namespace UE::AVCodecCore::H265
 		HighThroughput444_10 = HighThroughput,
 		HighThroughput444_14 = HighThroughput,
 		HighThroughput444_16Intra = HighThroughput,
+        MultiViewMain = 6,
+        ScalableMain = 7,
+        ThreeDimensionalMain = 8,
 		ScreenContentCoding = 9,
 		ScreenExtendedMain = ScreenContentCoding,
 		ScreenExtendedMain10 = ScreenContentCoding,
 		ScreenExtendedMain444 = ScreenContentCoding,
 		ScreenExtendedMain444_10 = ScreenContentCoding,
+        ScalableRangeExtensions = 10,
 		HighThroughputScreenContentCoding = 11,
 		ScreenExtendedHighThroughput444 = HighThroughputScreenContentCoding,
 		ScreenExtendedHighThroughput444_10 = HighThroughputScreenContentCoding,
@@ -295,48 +299,41 @@ namespace UE::AVCodecCore::H265
 		return OutBuf - OutBase;
 	}
 
-	struct FNaluH265 : public FNalu
+	struct FNaluH265
 	{
-		uint64 StartIdx = 0, Size = 0;
-		uint8 StartCodeSize = 3;
-		U<6, ENaluType> nal_unit_type = ENaluType::UNSPECIFIED;
-		U<6> nuh_layer_id = 0;
-		U<3> nuh_temporal_id_plus1 = 0;
-		const uint8* EBSP = nullptr;
-
-		bool IsSlice()
-		{
-			return	nal_unit_type == ENaluType::TRAIL_N ||
-					nal_unit_type == ENaluType::TRAIL_R ||
-					nal_unit_type == ENaluType::TSA_N ||
-					nal_unit_type == ENaluType::TSA_R ||
-					nal_unit_type == ENaluType::STSA_N ||
-					nal_unit_type == ENaluType::STSA_R ||
-					nal_unit_type == ENaluType::RADL_N ||
-					nal_unit_type == ENaluType::RADL_R ||
-					nal_unit_type == ENaluType::RASL_N ||
-					nal_unit_type == ENaluType::RASL_R ||
-					nal_unit_type == ENaluType::BLA_W_LP ||
-					nal_unit_type == ENaluType::BLA_W_RADL ||
-					nal_unit_type == ENaluType::BLA_N_LP ||
-					nal_unit_type == ENaluType::IDR_W_RADL ||
-					nal_unit_type == ENaluType::IDR_N_LP ||
-					nal_unit_type == ENaluType::CRA_NUT;
-			;
-		}
+		uint64 Start, Size;
+		uint8 StartCodeSize;
+        ENaluType Type;
+		uint8 NuhLayerId;
+        uint8 NuhTemporalIdPlus1;
+        const uint8* Data;
 	};
 
-	FAVResult FindNALUs(FVideoPacket const& InPacket, TArray<FNaluH265>& FoundNalus);
+	AVCODECSCORE_API FAVResult FindNALUs(FVideoPacket const& InPacket, TArray<FNaluH265>& FoundNalus);
 
 	struct profile_tier_level_t : public FBitstreamSegment
 	{
 		U<2> general_profile_space;
 		U<1> general_tier_flag;
 		U<5, EH265ProfileIDC> general_profile_idc;
-		U<32> general_profile_compatibility_flag;
+        TArray<U<1>> general_profile_compatibility_flag;
 		U<1> general_progressive_source_flag;
 		U<1> general_interlaced_source_flag;
-		EH265ConstraintFlag general_constraint_flags = EH265ConstraintFlag::None;
+        U<1> general_non_packed_constraint_flag;
+        U<1> general_frame_only_constraint_flag;
+
+
+        U<1> general_max_12bit_constraint_flag;
+        U<1> general_max_10bit_constraint_flag;
+        U<1> general_max_8bit_constraint_flag;
+        U<1> general_max_422chroma_constraint_flag;
+        U<1> general_max_420chroma_constraint_flag;
+        U<1> general_max_monochrome_constraint_flag;
+        U<1> general_intra_constraint_flag;
+        U<1> general_one_picture_only_constraint_flag;
+        U<1> general_lower_bit_rate_constraint_flag;
+        U<1> general_max_14bit_constraint_flag;
+
 		U<1> general_inbld_flag = 0;
 		U<8, EH265LevelIDC> general_level_idc;
 
@@ -350,15 +347,28 @@ namespace UE::AVCodecCore::H265
 			U<5, EH265ProfileIDC> sub_layer_profile_idc;
 			U<1> sub_layer_progressive_source_flag;
 			U<1> sub_layer_interlaced_source_flag;
-			U<32> sub_layer_profile_compatibility_flag;
-			EH265ConstraintFlag sub_layer_constraint_flags = EH265ConstraintFlag::None;
+            TArray<U<1>> sub_layer_profile_compatibility_flag;
+            U<1> sub_layer_non_packed_constraint_flag;
+            U<1> sub_layer_frame_only_constraint_flag;
+
+            U<1> sub_layer_max_12bit_constraint_flag;
+            U<1> sub_layer_max_10bit_constraint_flag;
+            U<1> sub_layer_max_8bit_constraint_flag;
+            U<1> sub_layer_max_422chroma_constraint_flag;
+            U<1> sub_layer_max_420chroma_constraint_flag;
+            U<1> sub_layer_max_monochrome_constraint_flag;
+            U<1> sub_layer_intra_constraint_flag;
+            U<1> sub_layer_one_picture_only_constraint_flag;
+            U<1> sub_layer_lower_bit_rate_constraint_flag;
+            U<1> sub_layer_max_14bit_constraint_flag;
+
 			U<1> sub_layer_inbld_flag;
 			U<8, EH265LevelIDC> sub_layer_level_idc;
 		};
 
 		TArray<sub_layer_t> sub_layers;
 
-		inline void Parse(const uint8& profilePresentFlag, uint8 maxNumSubLayersMinus1, FBitstreamReader& Bitstream);
+        inline void Parse(uint8 const& profilePresentFlag, uint8 maxNumSubLayersMinus1, FBitstreamReader& Bitstream);
 	};
 
 	struct hrd_parameters_t : public FBitstreamSegment
@@ -396,15 +406,15 @@ namespace UE::AVCodecCore::H265
 
 			TArray<sub_layer_hrd_parameters_t> sub_layer_hrd_parameters;
 
-			inline void Parse(uint8 const& in_sub_pic_hrd_params_present_flag, uint8 const& CpbCnt, FBitstreamReader& Bitstream);
-		};
+            inline void Parse(uint8 const& in_sub_pic_hrd_params_present_flag, uint8 const& CpbCnt, FBitstreamReader& Bitstream);
+        };
 
 		TArray<sub_layer_t> sub_layers;
 
-		inline void Parse(uint8 const& commonInfPresentFlag, uint8 const& maxNumSubLayersMinus1, FBitstreamReader& Bitstream);
-	};
+        inline void Parse(uint8 const& commonInfPresentFlag, uint8 const& maxNumSubLayersMinus1, FBitstreamReader& Bitstream);
+    };
 
-	struct FNaluVPS : public FNaluH265
+	struct VPS_t : public FNalu
 	{
 		U<4> vps_video_parameter_set_id;
 		U<1> vps_base_layer_internal_flag;
@@ -416,34 +426,31 @@ namespace UE::AVCodecCore::H265
 		profile_tier_level_t profile_tier_level;
 
 		U<1> vps_sub_layer_ordering_info_present_flag;
-		TArray<UE> vps_max_dec_pic_buffering_minus1 = {};
-		TArray<UE> vps_max_num_reorder_pics = {};
-		TArray<UE> vps_max_latency_increase_plus1 = {};
+		TArray<UE> vps_max_dec_pic_buffering_minus1;
+		TArray<UE> vps_max_num_reorder_pics;
+		TArray<UE> vps_max_latency_increase_plus1;
 		U<6> vps_max_layer_id;
 		UE vps_num_layer_sets_minus1;
-		TArray<TArray<U<1>>> layer_id_included_flag = {};
+		TArray<TArray<U<1>>> layer_id_included_flag;
 		U<1> vps_timing_info_present_flag;
 		U<32> vps_num_units_in_tick;
 		U<32> vps_time_scale;
 		U<1> vps_poc_proportional_to_timing_flag;
 		UE vps_num_ticks_poc_diff_one_minus1;
 		UE vps_num_hrd_parameters;
-		TArray<UE> hrd_layer_set_idx = {};
-		TArray<U<1>> cprms_present_flag = {};
-		TArray<hrd_parameters_t> HrdParameters = {};
+		TArray<UE> hrd_layer_set_idx;
+		TArray<U<1>> cprms_present_flag;
+		TArray<hrd_parameters_t> hrd_parameters;
 		U<1> vps_extension_flag;
 
 		// Stream State Members
 		// TODO (aidan) perhaps this should be moved elsewhere
 		mutable uint32 prevPicOrderCntLsb = 0;
 		mutable uint32 prevPicOrderCntMsb = 0;
+	};
 
-		FNaluVPS() = delete;
-		FNaluVPS(FNaluH265 const& InNaluH265)
-			: FNaluH265(InNaluH265) {}
+    FAVResult ParseVPS(FBitstreamReader& Bitstream, FNaluH265 const& InNaluInfo, TMap<uint32, VPS_t>& OutMapVPS);
 
-		FAVResult Parse();
-	};
 	
 	template <uint32 Tlog2blkSize, uint32 TscanIdx>
 	struct TScanOrder
@@ -556,7 +563,7 @@ public:
 		uint8 ScalingList0[6][16];
 		uint8 ScalingList1to3[3][6][64];
 
-		inline void Parse(FBitstreamReader& Bitstream);
+        inline void Parse(FBitstreamReader& Bitstream);
 	};
 
 	struct FH265ChromaInfo
@@ -611,15 +618,15 @@ public:
 		UE delta_idx_minus1 = 0;
 		U<1> delta_rps_sign;
 		UE abs_delta_rps_minus1;
-		TArray<U<1>> used_by_curr_pic_flags = {};
-		TArray<U<1>> use_delta_flags = {};
+		U<1>* used_by_curr_pic_flags;
+		U<1>* use_delta_flags;
 
 		UE num_negative_pics = 0;
 		UE num_positive_pics = 0;
-		TArray<UE> delta_poc_s0_minus1s = {};
-		TArray<U<1>> used_by_curr_pic_s0_flags = {};
-		TArray<UE> delta_poc_s1_minus1s = {};
-		TArray<U<1>> used_by_curr_pic_s1_flags = {};
+		UE* delta_poc_s0_minus1s;
+		U<1>* used_by_curr_pic_s0_flags;
+		UE* delta_poc_s1_minus1s;
+		U<1>* used_by_curr_pic_s1_flags;
 
 		// Derived Values (specification usage order)
 		uint8 NumPositivePics = 0;
@@ -630,19 +637,11 @@ public:
 		TStaticArray<uint8, 16> UsedByCurrPicS0;
 		TStaticArray<uint8, 16> UsedByCurrPicS1;
 
-		inline void Parse(uint8 const& stRpsIdx, TArray<short_term_ref_pic_set_t> const& short_term_ref_pic_sets, FBitstreamReader& Bitstream);
+        inline void Parse(uint8 const& stRpsIdx, TArray<short_term_ref_pic_set_t> const& short_term_ref_pic_sets, FBitstreamReader& Bitstream);
 		inline void CalculateValues(uint8 const& stRpsIdx, TArray<short_term_ref_pic_set_t>& st_ref_pic_set);
 	};
 
-	struct long_term_ref_pics_t : FBitstreamSegment
-	{
-		U<> lt_ref_pic_poc_lsb_sps;
-		U<1> used_by_curr_pic_lt_sps_flag;
-
-		inline void Parse(const uint8& log2_max_pic_order_cnt_lsb_minus4, FBitstreamReader& Bitstream);
-	};
-
-	struct FNaluSPS : public FNaluH265
+	struct SPS_t : public FNalu
 	{
 		U<4> sps_video_parameter_set_id;
 		U<3> sps_max_sub_layers_minus1;
@@ -650,7 +649,7 @@ public:
 		profile_tier_level_t profile_tier_level;
 		UE sps_seq_parameter_set_id;
 		UE chroma_format_idc;
-		U<1> separate_colour_plane_flag = 0;
+		U<1> separate_colour_plane_flag;
 		UE ChromaArrayType;
 		UE pic_width_in_luma_samples;
 		UE pic_height_in_luma_samples;
@@ -658,10 +657,10 @@ public:
 
 		struct conf_win_t : public FBitstreamSegment
 		{
-			UE left_offset = 0;
-			UE right_offset = 0;
-			UE top_offset = 0;
-			UE bottom_offset = 0;
+			UE left_offset;
+			UE right_offset;
+			UE top_offset;
+			UE bottom_offset;
 		} conf_win;
 
 		UE bit_depth_luma_minus8;
@@ -705,7 +704,8 @@ public:
 
 		U<1> long_term_ref_pics_present_flag;
 		UE num_long_term_ref_pics_sps;
-		TArray<long_term_ref_pics_t> long_term_ref_pics_sps;
+        TArray<U<>> lt_ref_pic_poc_lsb_sps = {};
+        TArray<U<1>> used_by_curr_pic_lt_sps_flag;
 
 		U<1> sps_temporal_mvp_enabled_flag;
 		U<1> strong_intra_smoothing_enabled_flag;
@@ -763,7 +763,7 @@ public:
 			UE log2_max_mv_length_horizontal;
 			UE log2_max_mv_length_vertical;
 
-			inline void Parse(FBitstreamReader& Bitstream);
+            inline void Parse(FBitstreamReader& Bitstream, uint8 const& sps_max_sub_layers_minus1);
 		} vui_parameters;
 
 		U<1> sps_extension_present_flag;
@@ -786,21 +786,19 @@ public:
 			U<1> persistent_rice_adaptation_enabled_flag;
 			U<1> cabac_bypass_alignment_enabled_flag;
 
-			inline void Parse(FBitstreamReader& Bitstream);
+            inline void Parse(FBitstreamReader& Bitstream);
 		} sps_range_extension;
 
 		struct sps_multilayer_extension_t : public FBitstreamSegment
 		{
 			// TODO (aidan)
-
-			inline void Parse(FBitstreamReader& Bitstream) { unimplemented(); }
+            inline void Parse(FBitstreamReader& Bitstream) { unimplemented(); }
 		} sps_multilayer_extension;
 
 		struct sps_3d_extension_t : public FBitstreamSegment
 		{
 			// TODO (aidan)
-
-			inline void Parse(FBitstreamReader& Bitstream) { unimplemented(); }
+            inline void Parse(FBitstreamReader& Bitstream) { unimplemented(); }
 		} sps_3d_extension;
 
 		struct sps_scc_extension_t : public FBitstreamSegment
@@ -811,24 +809,21 @@ public:
 			UE delta_palette_max_predictor_size;
 			U<1> sps_palette_predictor_initializers_present_flag;
 			UE sps_num_palette_predictor_initializers_minus1;
-			TArray<TArray<U<>>> sps_palette_predictor_initializers;
+            TArray<TArray<U<>>> sps_palette_predictor_initializers = {};
 			U<2> motion_vector_resolution_control_idc;
 			U<1> intra_boundary_filtering_disabled_flag;
 
-			inline void Parse(uint32 const& in_chroma_format_idc, uint32 const& in_bit_depth_chroma_minus8, FBitstreamReader& Bitstream);
+            inline void Parse(uint32 const& in_chroma_format_idc, uint32 const& in_bit_depth_chroma_minus8, FBitstreamReader& Bitstream);
 		} sps_scc_extension;
 
 		// Members derived from bitstream
 		uint32 MaxDpbSize;
+	};
 
-		FNaluSPS() = delete;
-		FNaluSPS(FNaluH265 const& InNaluH265)
-			: FNaluH265(InNaluH265) {}
+    FAVResult ParseSPS(FBitstreamReader& Bitstream, FNaluH265 const& InNaluInfo, TMap<uint32, VPS_t> const& InMapVPS, TMap<uint32, SPS_t>& OutMapSPS);
 
-		FAVResult Parse();
-	};
 
-	struct FNaluPPS : public FNaluH265
+	struct PPS_t : public FNalu
 	{
 		UE pps_pic_parameter_set_id;
 		UE pps_seq_parameter_set_id;
@@ -843,7 +838,7 @@ public:
 		U<1> constrained_intra_pred_flag;
 		U<1> transform_skip_enabled_flag;
 		U<1> cu_qp_delta_enabled_flag;
-		UE diff_cu_qp_delta_depth = 0;
+		UE diff_cu_qp_delta_depth;
 		SE pps_cb_qp_offset;
 		SE pps_cr_qp_offset;
 		U<1> pps_slice_chroma_qp_offsets_present_flag;
@@ -888,21 +883,17 @@ public:
 			UE log2_sao_offset_scale_luma;
 			UE log2_sao_offset_scale_chroma;
 
-			inline void Parse(uint8 const& in_transform_skip_enabled_flag, FBitstreamReader& Bitstream);
+            inline void Parse(uint8 const& in_transform_skip_enabled_flag, FBitstreamReader& Bitstream);
 		} pps_range_extension;
 
 		struct pps_multilayer_extension_t
 		{
-			// TODO (aidan)
-
-			inline void Parse(FBitstreamReader& Bitstream) { /*unimplemented();*/ }
+            inline void Parse(FBitstreamReader& Bitstream) { /*unimplemented();*/ }
 		} pps_multilayer_extension;
 
 		struct pps_3d_extension_t
 		{
-			// TODO (aidan)
-
-			inline void Parse(FBitstreamReader& Bitstream) { /*unimplemented();*/ }
+            inline void Parse(FBitstreamReader& Bitstream) { /*unimplemented();*/ }
 		} pps_3d_extension;
 
 		struct pps_scc_extension_t
@@ -918,18 +909,14 @@ public:
 			U<1> monochrome_palette_flag;
 			UE luma_bit_depth_entry_minus8;
 			UE chroma_bit_depth_entry_minus8;
-			TArray<TArray<U<>>> pps_palette_predictor_initializer;
+            TArray<TArray<U<>>> pps_palette_predictor_initializer = {};
 
-			inline void Parse(FBitstreamReader& Bitstream);
+            inline void Parse(FBitstreamReader& Bitstream);
 		} pps_scc_extension;
-
-		FNaluPPS() = delete;
-		FNaluPPS(FNaluH265 const& InNaluH265)
-			: FNaluH265(InNaluH265) {}
-
-		FAVResult Parse();
 	};
 
+    FAVResult ParsePPS(FBitstreamReader& Bitstream, FNaluH265 const& InNaluInfo, TMap<uint32, VPS_t> const& InMapVPS, TMap<uint32, SPS_t> const& InMapSPS, TMap<uint32, PPS_t>& OutMapPPS);
+
 	enum class EH265SliceType : uint8
 	{
 		B = 0,
@@ -937,15 +924,15 @@ public:
 		I = 2
 	};
 
-	struct FNaluSlice : public FNaluH265
+	struct Slice_t : public FNalu
 	{
 		// Members read from bitstream
 		U<1> first_slice_segment_in_pic_flag;
 		U<1> no_output_of_prior_pics_flag;
 		UE slice_pic_parameter_set_id;
 
-		U<1> dependent_slice_segment_flag;
-		U<> slice_segment_address;
+		U<1> dependent_slice_segment_flag = 0;
+		U<> slice_segment_address = 0;
 		UnsignedExpGolomb<EH265SliceType> slice_type;
 		U<1> pic_output_flag = 1;
 		U<2> colour_plane_id;
@@ -953,39 +940,38 @@ public:
 		U<1> short_term_ref_pic_set_sps_flag;
 		short_term_ref_pic_set_t short_term_ref_pic_set;
 		U<> short_term_ref_pic_set_idx = 0;
-		UE num_long_term_sps;
+		UE num_long_term_sps = 0;
 		UE num_long_term_pics = 0;
-		TArray<U<>> lt_idx_sps;
-		TArray<U<>> poc_lsb_lt;
-		TArray<U<1>> used_by_curr_pic_lt_flag;
-		TArray<U<1>> delta_poc_msb_present_flag;
-		TArray<UE> delta_poc_msb_cycle_lt;
+        U<>* lt_idx_sps = nullptr;
+        U<>* poc_lsb_lt = nullptr;
+		U<1>* used_by_curr_pic_lt_flag = nullptr;
+	    U<1>* delta_poc_msb_present_flag = nullptr;
+		UE* delta_poc_msb_cycle_lt = nullptr;
 		U<1> slice_temporal_mvp_enabled_flag = 0;
 
 		U<1> slice_sao_luma_flag = 0;
 		U<1> slice_sao_chroma_flag = 0;
 
 		U<1> num_ref_idx_active_override_flag;
-		UE num_ref_idx_l0_active_minus1;
-		UE num_ref_idx_l1_active_minus1;
+		UE num_ref_idx_l0_active_minus1 = 0;
+		UE num_ref_idx_l1_active_minus1 = 0;
 		struct ref_pic_list_modification_t
 		{
 			U<1> ref_pic_list_modification_flag_l0 = 0;
-			TArray<U<>> list_entry_l0;
+            U<>* list_entry_l0 = nullptr;
 			U<1> ref_pic_list_modification_flag_l1 = 0;
-			TArray<U<>> list_entry_l1;
-
-			inline void Parse(uint8 const& in_num_ref_idx_l0_active_minus1, uint8 const& in_num_ref_idx_l1_active_minus1, uint32 const& InNumPicTotalCurr, EH265SliceType const& SliceType, FBitstreamReader& BitStream);
+            U<>* list_entry_l1 = nullptr;
 
+            inline void Parse(uint8 const& in_num_ref_idx_l0_active_minus1, uint8 const& in_num_ref_idx_l1_active_minus1, uint32 const& InNumPicTotalCurr, EH265SliceType const& SliceType, FBitstreamReader& BitStream);
 		} ref_pic_list_modification;
-		U<1> mvd_l1_zero_flag;
-		U<1> cabac_init_flag;
+		U<1> mvd_l1_zero_flag = 0;
+		U<1> cabac_init_flag = 0;
 		U<1> collocated_from_l0_flag = 1;
 		UE collocated_ref_idx = 0;
 		struct pred_weight_table_t
 		{
 			UE luma_log2_weight_denom;
-			SE delta_chroma_log2_weight_denom;
+			SE delta_chroma_log2_weight_denom = 0;
 
 			TStaticArray<U<1>, 15> luma_weight_l0_flag;
 			TStaticArray<U<1>, 15> chroma_weight_l0_flag;
@@ -1001,8 +987,8 @@ public:
 			TStaticArray<TStaticArray<SE, 2>, 15> delta_chroma_weight_l1;
 			TStaticArray<TStaticArray<SE, 2>, 15> delta_chroma_offset_l1;
 
-			inline void Parse(const uint8& ChromaArrayType, FNaluSlice const& CurrentSlice, FBitstreamReader& BitStream);
-		} pred_weight_table;
+            inline void Parse(const uint8& ChromaArrayType, FNaluH265 const& InNaluInfo, Slice_t const& CurrentSlice, FBitstreamReader& BitStream);
+        } pred_weight_table;
 
 		UE five_minus_max_num_merge_cand;
 		U<1> use_integer_mv_flag;
@@ -1012,9 +998,9 @@ public:
 		SE slice_cb_qp_offset = 0;
 		SE slice_cr_qp_offset = 0;
 
-		SE slice_act_y_qp_offset = 0;
-		SE slice_act_cb_qp_offset = 0;
-		SE slice_act_cr_qp_offset = 0;
+		SE slice_act_y_qp_offset;
+		SE slice_act_cb_qp_offset;
+		SE slice_act_cr_qp_offset;
 
 		U<1> cu_chroma_qp_offset_enabled_flag = 0;
 
@@ -1028,37 +1014,23 @@ public:
 
 		UE num_entry_point_offsets = 0;
 		UE offset_len_minus1;
-		TArray<U<>> entry_point_offset_minus1;
-
-		UE slice_segment_header_extension_length;
-		TArray<U<8>> slice_segment_header_extension_data_byte;
+        U<>* entry_point_offset_minus1;
 
-		// Members derived from bitstream
-		TWeakPtr<FNaluPPS> CurrentPPS;
-		TWeakPtr<FNaluSPS> CurrentSPS;
-		TWeakPtr<FNaluVPS> CurrentVPS;
+		UE slice_segment_header_extension_length = 0;
+		U<8>* slice_segment_header_extension_data_byte;
 
 		uint32 CurrPicIdx = 0;
 		uint32 CurrPicOrderCntVal = 0;
 		uint32 CurrRpsIdx = 0;
-		uint32 NumBitsForShortTermRPSInSlice = 0; // Something for NVDEC
-		
-		FNaluSlice() = delete;
-		FNaluSlice(FNaluH265 const& InNaluH265)
-			: FNaluH265(InNaluH265) {}
-
-		FAVResult Parse(FVideoDecoderConfigH265* InConfig);
-		bool IsIDR() { return  nal_unit_type == ENaluType::IDR_N_LP || nal_unit_type == ENaluType::IDR_W_RADL; };
 	};
 
-	struct FNaluSEI : public FNaluH265
-	{
+    FAVResult ParseSliceHeader(FBitstreamReader& Bitstream, FNaluH265 const& InNaluInfo, TMap<uint32, VPS_t> const& InMapVPS, TMap<uint32, SPS_t> const& InMapSPS, TMap<uint32, PPS_t> const& InMapPPS, Slice_t& OutSlice);
 
-		FNaluSEI() = delete;
-		FNaluSEI(FNaluH265 const& InNaluH265)
-			: FNaluH265(InNaluH265) {}
 
-		FAVResult Parse();
+	struct SEI_t : public FNalu
+	{
 	};
 
+    FAVResult ParseSEI(FBitstreamReader& Bitstream, FNaluH265 const& InNaluInfo, SEI_t& OutSEI);
+
 } // namespace UE::AVCodecCore::H265
diff --git a/Engine/Plugins/Experimental/AVCodecs/AVCodecsCore/Source/AVCodecsCore/Public/Video/Decoders/Configs/VideoDecoderConfigH264.h b/Engine/Plugins/Experimental/AVCodecs/AVCodecsCore/Source/AVCodecsCore/Public/Video/Decoders/Configs/VideoDecoderConfigH264.h
--- a/Engine/Plugins/Experimental/AVCodecs/AVCodecsCore/Source/AVCodecsCore/Public/Video/Decoders/Configs/VideoDecoderConfigH264.h
+++ b/Engine/Plugins/Experimental/AVCodecs/AVCodecsCore/Source/AVCodecsCore/Public/Video/Decoders/Configs/VideoDecoderConfigH264.h
@@ -19,7 +19,9 @@ struct AVCODECSCORE_API FVideoDecoderConfigH264 : public FVideoDecoderConfig
 	{
 	}
 
-	FAVResult Parse(TSharedRef<FAVInstance> const& Instance, FVideoPacket const& Packet, TArray<UE::AVCodecCore::H264::Slice_t>& OutSlices);
+	FAVResult Parse(FVideoPacket const& Packet, TArray<UE::AVCodecCore::H264::Slice_t>& OutSlices);
+
+    TOptional<int> GetLastSliceQP(TArray<UE::AVCodecCore::H264::Slice_t>& Slices);
 };
 
 DECLARE_TYPEID(FVideoDecoderConfigH264, AVCODECSCORE_API);
diff --git a/Engine/Plugins/Experimental/AVCodecs/AVCodecsCore/Source/AVCodecsCore/Public/Video/Decoders/Configs/VideoDecoderConfigH265.h b/Engine/Plugins/Experimental/AVCodecs/AVCodecsCore/Source/AVCodecsCore/Public/Video/Decoders/Configs/VideoDecoderConfigH265.h
--- a/Engine/Plugins/Experimental/AVCodecs/AVCodecsCore/Source/AVCodecsCore/Public/Video/Decoders/Configs/VideoDecoderConfigH265.h
+++ b/Engine/Plugins/Experimental/AVCodecs/AVCodecsCore/Source/AVCodecsCore/Public/Video/Decoders/Configs/VideoDecoderConfigH265.h
@@ -5,111 +5,24 @@
 #include "Video/VideoDecoder.h"
 #include "Video/CodecUtils/CodecUtilsH265.h"
 
-enum class EPictureUsage : uint8
-{
-	UNUSED			= 0u,
-	SHORT_TERM		= 1u << 0,
-	LONG_TERM		= 1u << 1,		
-};
-
-ENUM_CLASS_FLAGS(EPictureUsage);
-
 /*
  * Configuration settings for H265 decoders.
  */
-struct FVideoDecoderConfigH265 : public FVideoDecoderConfig
+struct AVCODECSCORE_API FVideoDecoderConfigH265 : public FVideoDecoderConfig
 {
+    TMap<uint32, UE::AVCodecCore::H265::VPS_t> VPS;
+	TMap<uint32, UE::AVCodecCore::H265::SPS_t> SPS;
+	TMap<uint32, UE::AVCodecCore::H265::PPS_t> PPS;
+	TArray<UE::AVCodecCore::H265::SEI_t> SEI;
+
 	FVideoDecoderConfigH265(EAVPreset Preset = EAVPreset::Default)
 		: FVideoDecoderConfig(Preset)
 	{
 	}
 
-	AVCODECSCORE_API FAVResult Parse(TSharedRef<FAVInstance> const& Instance, FVideoPacket const& Packet, TArray<TSharedPtr<UE::AVCodecCore::H265::FNaluH265>>& Slices);
-
-	AVCODECSCORE_API FAVResult UpdateScalingLists(TSharedRef<UE::AVCodecCore::H265::FNaluSlice> CurrentSlice);
-
-	AVCODECSCORE_API FAVResult UpdateRPS(TSharedRef<UE::AVCodecCore::H265::FNaluSlice> CurrentSlice);
-
-	TMap<uint32, TSharedPtr<UE::AVCodecCore::H265::FNaluVPS>> ParsedVPS = {};
-	TMap<uint32, TSharedPtr<UE::AVCodecCore::H265::FNaluSPS>> ParsedSPS = {};
-	TMap<uint32, TSharedPtr<UE::AVCodecCore::H265::FNaluPPS>> ParsedPPS = {};
-	TArray<TSharedPtr<UE::AVCodecCore::H265::FNaluSEI>> ParsedSEI = {};
-
-	// Decoder state
-	int32 CurrPicIdx = -1;
-	int32 LastRpsIdx = 0;
-	int32 LastPicOrderCntValue = 0;
-
-	// Scaling lists
-	uint8 ScalingList4x4[6][16];
-	uint8 ScalingList8x8[6][16];
-	uint8 ScalingList16x16[6][16];
-	uint8 ScalingList32x32[6][16];
-
-	int16 ScalingListDCCoeff4x4[6];
-	int16 ScalingListDCCoeff8x8[6];
-	int16 ScalingListDCCoeff16x16[6];
-	int16 ScalingListDCCoeff32x32[6];
-	
-	struct FReferencePictureSet
-	{
-		uint8 NumPicTotalCurr;
-		
-		int32 PicOrderCntVal[16] = { 0 };
-		uint8 PicAvaliability[16] = { 0 };
-		EPictureUsage PicUsage[16]  = { EPictureUsage::UNUSED };
-		uint8 PicLayerId[16] = { 0 };
-		
-		uint8 NumPocStCurrBefore;
-		int32 PocStCurrBefore[16] = { 0 };
-		uint8 RefPicSetStCurrBefore[8] = { 0 };
-		
-		uint8 NumPocStCurrAfter;
-		int32 PocStCurrAfter[16] = { 0 };
-		uint8 RefPicSetStCurrAfter[8] = { 0 };
-
-		uint8 NumPocStFoll;
-		int32 PocStFoll[16] = { 0 };
-		uint8 RefPicSetStFoll[16] = { 0 };
-
-		uint8 NumPocLtCurr;
-		int32 PocLtCurr[16] = { 0 };
-		uint8 RefPicSetLtCurr[16] = { 0 };
-
-		uint8 NumPocLtFoll;
-		int32 PocLtFoll[16] = { 0 };
-		uint8 RefPicSetLtFoll[16] = { 0 };
-
-		uint8 RefPicList0[16] = { 0 };
-		uint8 RefPicList1[16] = { 0 };
-
-		uint8 RefPicSetInterLayer0[8] = { 0 };
-		uint8 RefPicSetInterLayer1[8] = { 0 };
-
-		void PrepState()
-		{
-			NumPocStCurrBefore = 0;
-			FMemory::Memset( PocStCurrBefore, 0, sizeof(PocStCurrBefore));
-			
-			NumPocStCurrAfter = 0;
-			FMemory::Memset( PocStCurrAfter, 0, sizeof(PocStCurrAfter));
-			
-			NumPocStFoll = 0;
-			FMemory::Memset( PocStFoll, 0, sizeof(PocStFoll));
-			
-			NumPocLtCurr = 0;
-			FMemory::Memset( PocLtCurr, 0, sizeof(PocLtCurr));
-			
-			NumPocLtFoll = 0;
-			FMemory::Memset( PocLtFoll, 0, sizeof(PocLtFoll));
-
-			NumPicTotalCurr = 0;
-
-			FMemory::Memset( PicUsage, 0, sizeof(PicUsage));
+	FAVResult Parse(FVideoPacket const& Packet, TArray<UE::AVCodecCore::H265::Slice_t>& OutSlices);
 
-			FMemory::Memset( PicOrderCntVal, 0, sizeof(PicOrderCntVal));
-		}
-	} ReferencePictureSet;
+    TOptional<int> GetLastSliceQP(TArray<UE::AVCodecCore::H265::Slice_t>& Slices);
 };
 
 DECLARE_TYPEID(FVideoDecoderConfigH265, AVCODECSCORE_API);
diff --git a/Engine/Plugins/Experimental/AVCodecs/AVCodecsCore/Source/AVCodecsCore/Public/Video/Decoders/Configs/VideoDecoderConfigVP9.h b/Engine/Plugins/Experimental/AVCodecs/AVCodecsCore/Source/AVCodecsCore/Public/Video/Decoders/Configs/VideoDecoderConfigVP9.h
new file mode 100644
--- /dev/null
+++ b/Engine/Plugins/Experimental/AVCodecs/AVCodecsCore/Source/AVCodecsCore/Public/Video/Decoders/Configs/VideoDecoderConfigVP9.h
@@ -0,0 +1,21 @@
+// Copyright Epic Games, Inc. All Rights Reserved.
+
+#pragma once
+
+#include "Video/VideoDecoder.h"
+#include "Video/CodecUtils/CodecUtilsVP9.h"
+
+/*
+ * Configuration settings for VP9 decoders.
+ */
+struct AVCODECSCORE_API FVideoDecoderConfigVP9 : public FVideoDecoderConfig
+{
+    FVideoDecoderConfigVP9(EAVPreset Preset = EAVPreset::Default)
+		: FVideoDecoderConfig(Preset)
+	{
+	}
+
+	FAVResult Parse(FVideoPacket const& Packet, UE::AVCodecCore::VP9::Header_t& OutHeader);
+};
+
+DECLARE_TYPEID(FVideoDecoderConfigVP9, AVCODECSCORE_API);
diff --git a/Engine/Plugins/Experimental/AVCodecs/AVCodecsCore/Source/AVCodecsCore/Public/Video/Resources/D3D/VideoResourceD3D.h b/Engine/Plugins/Experimental/AVCodecs/AVCodecsCore/Source/AVCodecsCore/Public/Video/Resources/D3D/VideoResourceD3D.h
new file mode 100644
--- /dev/null
+++ b/Engine/Plugins/Experimental/AVCodecs/AVCodecsCore/Source/AVCodecsCore/Public/Video/Resources/D3D/VideoResourceD3D.h
@@ -0,0 +1,109 @@
+// Copyright Epic Games, Inc. All Rights Reserved.
+
+#pragma once
+
+#include "Templates/RefCounting.h"
+
+#include "AVContext.h"
+#include "Video/VideoResource.h"
+
+#pragma warning(push)
+#pragma warning(disable: 4005)
+
+THIRD_PARTY_INCLUDES_START
+#include "Windows/AllowWindowsPlatformTypes.h"
+#include <d3d11.h>
+#include <mfobjects.h>
+#include <mftransform.h>
+#include <mfapi.h>
+#include <mferror.h>
+#include <mfidl.h>
+#include <codecapi.h>
+#include <shlwapi.h>
+#include <mfreadwrite.h>
+#include <d3d11_1.h>
+#include <d3d12.h>
+#include <dxgi1_4.h>
+#include "Windows/HideWindowsPlatformTypes.h"
+THIRD_PARTY_INCLUDES_END
+
+/**
+ * D3D11 and D3D12 platform video context and resource.
+ */
+
+class AVCODECSCORE_API FVideoContextD3D11 : public FAVContext
+{
+public:
+	TRefCountPtr<ID3D11Device> Device;
+
+	FVideoContextD3D11(TRefCountPtr<ID3D11Device> const& Device);
+};
+
+class AVCODECSCORE_API FVideoContextD3D12 : public FAVContext
+{
+public:
+	TRefCountPtr<ID3D12Device> Device;
+
+	FVideoContextD3D12(TRefCountPtr<ID3D12Device> const& Device);
+};
+
+class AVCODECSCORE_API FVideoResourceD3D11 : public TVideoResource<FVideoContextD3D11>
+{
+private:
+	TRefCountPtr<ID3D11Texture2D> Raw;
+	HANDLE SharedHandle = nullptr;
+
+public:
+	static FVideoDescriptor GetDescriptorFrom(TSharedRef<FAVDevice> const& Device, TRefCountPtr<ID3D11Texture2D> const& Raw);
+
+	FORCEINLINE TRefCountPtr<ID3D11Texture2D> const& GetRaw() const { return Raw; }
+	FORCEINLINE HANDLE const& GetSharedHandle() const { return SharedHandle; }
+
+	FVideoResourceD3D11(TSharedRef<FAVDevice> const& Device, TRefCountPtr<ID3D11Texture2D> const& Raw, FAVLayout const& Layout);
+
+	virtual FAVResult Validate() const override;
+};
+
+class AVCODECSCORE_API FVideoResourceD3D12 : public TVideoResource<FVideoContextD3D12>
+{
+public:
+	struct FRawD3D12
+	{
+		TRefCountPtr<ID3D12Resource> D3DResource;
+		HANDLE D3DResourceShared = nullptr;
+		TRefCountPtr<ID3D12Heap> D3DHeap;
+		HANDLE D3DHeapShared = nullptr;
+		TRefCountPtr<ID3D12Fence> D3DFence;
+		HANDLE D3DFenceShared = nullptr;
+		uint64 FenceValue;
+	};
+
+private:
+	FRawD3D12 Raw;
+
+public:
+	static FVideoDescriptor GetDescriptorFrom(TSharedRef<FAVDevice> const& Device, FRawD3D12 const& Raw);
+
+	FORCEINLINE TRefCountPtr<ID3D12Resource> const& GetResource() const { return Raw.D3DResource; }
+	FORCEINLINE HANDLE const& GetResourceSharedHandle() const { return Raw.D3DResourceShared; }
+	
+	FORCEINLINE TRefCountPtr<ID3D12Heap> const& GetHeap() const { return Raw.D3DHeap; }
+	FORCEINLINE HANDLE const& GetHeapSharedHandle() const { return Raw.D3DHeapShared; }
+
+	FORCEINLINE TRefCountPtr<ID3D12Fence> const& GetFence() const { return Raw.D3DFence; };
+	FORCEINLINE HANDLE const& GetFenceSharedHandle() const { return Raw.D3DFenceShared; }
+	FORCEINLINE uint64 const& GetFenceValue() const { return Raw.FenceValue; }
+
+	uint64 GetSizeInBytes();
+
+	FVideoResourceD3D12(TSharedRef<FAVDevice> const& Device, FRawD3D12 const& Raw, FAVLayout const& Layout);
+	FVideoResourceD3D12(TSharedRef<FAVDevice> const& Device, FRawD3D12 const& Raw, FAVLayout const& Layout, FVideoDescriptor const& Descriptor);
+	virtual ~FVideoResourceD3D12() override;
+
+	virtual FAVResult Validate() const override;
+};
+
+DECLARE_TYPEID(FVideoContextD3D11, AVCODECSCORE_API);
+DECLARE_TYPEID(FVideoContextD3D12, AVCODECSCORE_API);
+DECLARE_TYPEID(FVideoResourceD3D11, AVCODECSCORE_API);
+DECLARE_TYPEID(FVideoResourceD3D12, AVCODECSCORE_API);
\ No newline at end of file
diff --git a/Engine/Plugins/Experimental/AVCodecs/AVCodecsCore/Source/AVCodecsCore/Public/Video/Resources/Metal/VideoResourceMetal.h b/Engine/Plugins/Experimental/AVCodecs/AVCodecsCore/Source/AVCodecsCore/Public/Video/Resources/Metal/VideoResourceMetal.h
new file mode 100644
--- /dev/null
+++ b/Engine/Plugins/Experimental/AVCodecs/AVCodecsCore/Source/AVCodecsCore/Public/Video/Resources/Metal/VideoResourceMetal.h
@@ -0,0 +1,48 @@
+// Copyright Epic Games, Inc. All Rights Reserved.
+
+#if AVCODECS_USE_METAL
+
+#pragma once
+
+#include "Templates/RefCounting.h"
+
+#include "AVContext.h"
+#include "Video/VideoResource.h"
+
+THIRD_PARTY_INCLUDES_START
+#include "mtlpp.hpp"
+THIRD_PARTY_INCLUDES_END
+
+
+/**
+ * Metal platform video context and resource.
+ */
+
+class AVCODECSCORE_API FVideoContextMetal : public FAVContext
+{
+public:
+	mtlpp::Device Device;
+
+	FVideoContextMetal(mtlpp::Device const& Device);
+};
+
+
+class AVCODECSCORE_API FVideoResourceMetal : public TVideoResource<FVideoContextMetal>
+{
+private:
+	mtlpp::Texture* Raw;
+
+public:
+	static FVideoDescriptor GetDescriptorFrom(TSharedRef<FAVDevice> const& Device, mtlpp::Texture* Raw);
+
+	FORCEINLINE mtlpp::Texture* GetRaw() const { return Raw; }
+
+	FVideoResourceMetal(TSharedRef<FAVDevice> const& Device, mtlpp::Texture* Raw, FAVLayout const& Layout);
+
+	virtual FAVResult Validate() const override;
+};
+
+DECLARE_TYPEID(FVideoContextMetal, AVCODECSCORE_API);
+DECLARE_TYPEID(FVideoResourceMetal, AVCODECSCORE_API);
+
+#endif
diff --git a/Engine/Plugins/Experimental/AVCodecs/AVCodecsCore/Source/AVCodecsCore/Public/Video/Resources/Vulkan/VideoResourceVulkan.h b/Engine/Plugins/Experimental/AVCodecs/AVCodecsCore/Source/AVCodecsCore/Public/Video/Resources/Vulkan/VideoResourceVulkan.h
new file mode 100644
--- /dev/null
+++ b/Engine/Plugins/Experimental/AVCodecs/AVCodecsCore/Source/AVCodecsCore/Public/Video/Resources/Vulkan/VideoResourceVulkan.h
@@ -0,0 +1,61 @@
+// Copyright Epic Games, Inc. All Rights Reserved.
+
+#pragma once
+
+#include "AVContext.h"
+#include "Video/VideoResource.h"
+
+THIRD_PARTY_INCLUDES_START
+#if PLATFORM_WINDOWS
+#include "Windows/AllowWindowsPlatformTypes.h"
+#endif
+
+#include "vulkan_core.h"
+
+#if PLATFORM_WINDOWS
+#include "Windows/HideWindowsPlatformTypes.h"
+#endif
+THIRD_PARTY_INCLUDES_END
+
+/**
+ * Vulkan platform video context and resource.
+ */
+
+class AVCODECSCORE_API FVideoContextVulkan : public FAVContext
+{
+public:
+	VkInstance Instance;
+	VkDevice Device;
+	VkPhysicalDevice PhysicalDevice;
+	TFunction<PFN_vkVoidFunction(const char*)> vkGetDeviceProcAddr;
+
+	FVideoContextVulkan(VkInstance const& Instance, VkDevice const& Device, VkPhysicalDevice const& PhysicalDevice, TFunction<PFN_vkVoidFunction(const char*)> const& GetDeviceProcAddrFunc);
+};
+
+#if PLATFORM_WINDOWS
+typedef HANDLE FVulkanSharedHandle;
+#else
+typedef int FVulkanSharedHandle;
+#endif
+
+class AVCODECSCORE_API FVideoResourceVulkan : public TVideoResource<FVideoContextVulkan>
+{
+private:
+	VkDeviceMemory DeviceMemory;
+	FVulkanSharedHandle DeviceMemorySharedHandle = FVulkanSharedHandle();
+	VkFence Fence;
+	FVulkanSharedHandle FenceSharedHandle = FVulkanSharedHandle();
+
+public:
+	FORCEINLINE VkDeviceMemory GetRaw() const { return DeviceMemory; }
+	FORCEINLINE FVulkanSharedHandle const& GetSharedHandle() const { return DeviceMemorySharedHandle; }
+	FORCEINLINE FVulkanSharedHandle const& GetFenceSharedHandle() const { return FenceSharedHandle; }
+
+	FVideoResourceVulkan(TSharedRef<FAVDevice> const& Device, VkDeviceMemory Raw, FAVLayout const& Layout, FVideoDescriptor const& Descriptor);
+	virtual ~FVideoResourceVulkan() override = default;
+
+	virtual FAVResult Validate() const override;
+};
+
+DECLARE_TYPEID(FVideoContextVulkan, AVCODECSCORE_API);
+DECLARE_TYPEID(FVideoResourceVulkan, AVCODECSCORE_API);
diff --git a/Engine/Plugins/Experimental/AVCodecs/AVCodecsCore/Source/AVCodecsCoreRHI/AVCodecsCoreRHI.Build.cs b/Engine/Plugins/Experimental/AVCodecs/AVCodecsCore/Source/AVCodecsCoreRHI/AVCodecsCoreRHI.Build.cs
--- a/Engine/Plugins/Experimental/AVCodecs/AVCodecsCore/Source/AVCodecsCoreRHI/AVCodecsCoreRHI.Build.cs
+++ b/Engine/Plugins/Experimental/AVCodecs/AVCodecsCore/Source/AVCodecsCoreRHI/AVCodecsCoreRHI.Build.cs
@@ -15,18 +15,25 @@ public class AVCodecsCoreRHI : ModuleRules
 			"Engine",
 			"AVCodecsCore",
 			"RHI",
-			"VulkanRHI",
+			"RHICore",
 			"ColorManagement"
 		});
 
 		PublicDependencyModuleNames.AddRange(new string[] {
 			"RenderCore",
 			"Core",
-			"CoreUObject",
+			"CoreUObject"
 		});
 		
-		PrivateDependencyModuleNames.Add("Vulkan");
-		AddEngineThirdPartyPrivateStaticDependencies(Target, "Vulkan");
+		if (Target.IsInPlatformGroup(UnrealPlatformGroup.Windows) || Target.IsInPlatformGroup(UnrealPlatformGroup.Unix))
+		{
+			PrivateDependencyModuleNames.AddRange(new string[] {
+				"Vulkan",
+				"VulkanRHI"
+			});
+			
+			AddEngineThirdPartyPrivateStaticDependencies(Target, "Vulkan");
+		}
 
 		if (Target.IsInPlatformGroup(UnrealPlatformGroup.Windows))
 		{
@@ -39,5 +46,19 @@ public class AVCodecsCoreRHI : ModuleRules
 			AddEngineThirdPartyPrivateStaticDependencies(Target, "DX11");
 			AddEngineThirdPartyPrivateStaticDependencies(Target, "DX12");
 		}
+
+		if (Target.IsInPlatformGroup(UnrealPlatformGroup.Apple))
+		{
+			PrivateDependencyModuleNames.AddRange(new string[] {
+        	    "MetalRHI",
+        	});
+	
+        	PublicFrameworks.AddRange(new string[]{
+        	    "AVFoundation",
+        	    "VideoToolbox"
+        	});
+	
+        	AddEngineThirdPartyPrivateStaticDependencies(Target, "MTLPP");	
+		}
 	}
-}
\ No newline at end of file
+}
diff --git a/Engine/Plugins/Experimental/AVCodecs/AVCodecsCore/Source/AVCodecsCoreRHI/Private/AVCodecsCoreRHIModule.cpp b/Engine/Plugins/Experimental/AVCodecs/AVCodecsCore/Source/AVCodecsCoreRHI/Private/AVCodecsCoreRHIModule.cpp
--- a/Engine/Plugins/Experimental/AVCodecs/AVCodecsCore/Source/AVCodecsCoreRHI/Private/AVCodecsCoreRHIModule.cpp
+++ b/Engine/Plugins/Experimental/AVCodecs/AVCodecsCore/Source/AVCodecsCoreRHI/Private/AVCodecsCoreRHIModule.cpp
@@ -2,23 +2,29 @@
 
 #include "Modules/ModuleManager.h"
 
-#include "IVulkanDynamicRHI.h"
 #include "Misc/App.h"
 
 #include "Video/Resources/VideoResourceRHI.h"
-#include "Video/Resources/VideoResourceVulkan.h"
 #include "Video/Decoders/VideoDecoderRHI.h"
+#include "Video/Decoders/Configs/VideoDecoderConfigVP9.h"
 #include "Video/Decoders/Configs/VideoDecoderConfigH264.h"
 #include "Video/Decoders/Configs/VideoDecoderConfigH265.h"
 #include "Video/Encoders/VideoEncoderRHI.h"
 #include "Video/Encoders/Configs/VideoEncoderConfigH264.h"
 #include "Video/Encoders/Configs/VideoEncoderConfigH265.h"
 
-#if PLATFORM_WINDOWS
-#include "ID3D11DynamicRHI.h"
-#include "ID3D12DynamicRHI.h"
-
-#include "Video/Resources/Windows/VideoResourceD3D.h"
+#if AVCODECS_USE_D3D
+	#include "ID3D11DynamicRHI.h"
+	#include "ID3D12DynamicRHI.h"
+	#include "Video/Resources/D3D/VideoResourceD3D.h"
+#endif 
+#if AVCODECS_USE_VULKAN
+	#include "IVulkanDynamicRHI.h"
+	#include "Video/Resources/Vulkan/VideoResourceVulkan.h"
+#endif
+#if AVCODECS_USE_METAL
+	#include "DynamicRHI.h"
+	#include "Video/Resources/Metal/VideoResourceMetal.h"
 #endif
 
 class FAVCodecsCoreRHI : public IModuleInterface
@@ -28,6 +34,8 @@ public:
 	{
 		if (FApp::CanEverRender())
 		{
+#if AVCODECS_USE_VULKAN
+            
 #if PLATFORM_WINDOWS
 			TCHAR const* DynamicRHIModuleName = GetSelectedDynamicRHIModuleName(false);
 #elif PLATFORM_LINUX
@@ -44,11 +52,12 @@ public:
 
 				IVulkanDynamicRHI::AddEnabledDeviceExtensionsAndLayers(ExtensionsToAdd, TArray<ANSICHAR const*>());
 			}
-
+#endif // AVCODECS_USE_VULKAN
 			FCoreDelegates::OnPostEngineInit.AddLambda([]()
 			{
 				switch (GDynamicRHI->GetInterfaceType())
 				{
+#if AVCODECS_USE_VULKAN
 				case ERHIInterfaceType::Vulkan:
 					FAVDevice::GetHardwareDevice()->SetContext<FVideoContextVulkan>(
 						MakeShared<FVideoContextVulkan>(
@@ -60,7 +69,8 @@ public:
 							}));
 				
 					break;
-#if PLATFORM_WINDOWS
+#endif
+#if AVCODECS_USE_D3D
 				case ERHIInterfaceType::D3D11:
 					FAVDevice::GetHardwareDevice()->SetContext<FVideoContextD3D11>(
 						MakeShared<FVideoContextD3D11>(
@@ -73,6 +83,13 @@ public:
 							GetID3D12DynamicRHI()->RHIGetDevice(0)));
 				
 					break;
+#endif
+#if AVCODECS_USE_METAL
+                case ERHIInterfaceType::Metal:
+                    FAVDevice::GetHardwareDevice()->SetContext<FVideoContextMetal>(
+                        MakeShared<FVideoContextMetal>(
+                            *static_cast<mtlpp::Device*>(GDynamicRHI->RHIGetNativeDevice())));
+                    break;
 #endif
 				default:
 					break;
@@ -87,6 +104,7 @@ public:
 
 		FVideoDecoder::Register<TVideoDecoderRHI<FVideoDecoderConfigH264>, FVideoResourceRHI, FVideoDecoderConfigH264>();
 		FVideoDecoder::Register<TVideoDecoderRHI<FVideoDecoderConfigH265>, FVideoResourceRHI, FVideoDecoderConfigH265>();
+        FVideoDecoder::Register<TVideoDecoderRHI<FVideoDecoderConfigVP9>, FVideoResourceRHI, FVideoDecoderConfigVP9>();
 	}
 };
 
diff --git a/Engine/Plugins/Experimental/AVCodecs/AVCodecsCore/Source/AVCodecsCoreRHI/Private/Video/Resources/VideoResourceRHI.cpp b/Engine/Plugins/Experimental/AVCodecs/AVCodecsCore/Source/AVCodecsCoreRHI/Private/Video/Resources/VideoResourceRHI.cpp
--- a/Engine/Plugins/Experimental/AVCodecs/AVCodecsCore/Source/AVCodecsCoreRHI/Private/Video/Resources/VideoResourceRHI.cpp
+++ b/Engine/Plugins/Experimental/AVCodecs/AVCodecsCore/Source/AVCodecsCoreRHI/Private/Video/Resources/VideoResourceRHI.cpp
@@ -4,17 +4,22 @@
 
 #include "Video/Resources/VideoResourceCPU.h"
 
-#if PLATFORM_WINDOWS
+#if AVCODECS_USE_D3D
 	#include "ID3D11DynamicRHI.h"
 	#include "ID3D12DynamicRHI.h"
 	#include "D3D12RHI.h"
 	#include "BoundShaderStateCache.h"
 	#include "D3D12ShaderResources.h"
-	#include "Video/Resources/Windows/VideoResourceD3D.h"
+	#include "Video/Resources/D3D/VideoResourceD3D.h"
+#endif 
+#if AVCODECS_USE_VULKAN
+	#include "IVulkanDynamicRHI.h"
+	#include "Video/Resources/Vulkan/VideoResourceVulkan.h"
+#endif
+#if AVCODECS_USE_METAL
+	#include "Video/Resources/Metal/VideoResourceMetal.h"
 #endif
 
-#include "IVulkanDynamicRHI.h"
-#include "Video/Resources/VideoResourceVulkan.h"
 #include "Async/Async.h"
 #include "HAL/Event.h"
 #include "RHIStaticStates.h"
@@ -30,6 +35,7 @@ FAVLayout FVideoResourceRHI::GetLayoutFrom(TSharedRef<FAVDevice> const& Device,
 {
 	switch (GDynamicRHI->GetInterfaceType())
 	{
+#if AVCODECS_USE_VULKAN
 	case ERHIInterfaceType::Vulkan:
 		{
 			FRHITextureDesc const& Desc = Raw->GetDesc();
@@ -51,7 +57,8 @@ FAVLayout FVideoResourceRHI::GetLayoutFrom(TSharedRef<FAVDevice> const& Device,
 		}
 		
 		break;
-#if PLATFORM_WINDOWS
+#endif
+#if AVCODECS_USE_D3D
 	case ERHIInterfaceType::D3D11:
 		{
 			// This is a guess because D3D11 doesn't expose the actual allocation size
@@ -91,6 +98,21 @@ FAVLayout FVideoResourceRHI::GetLayoutFrom(TSharedRef<FAVDevice> const& Device,
 		}
 		
 		break;
+#endif
+#if AVCODECS_USE_METAL
+        case ERHIInterfaceType::Metal:
+            {
+                uint32 const Size =  GDynamicRHI->RHIComputeMemorySize(Raw);
+
+                uint32 const BlockSizeX = GPixelFormats[Raw->GetFormat()].BlockSizeX;
+                uint32 const BlockBytes = GPixelFormats[Raw->GetFormat()].BlockBytes;
+                uint32 const NumBlocksX = (Raw->GetSizeX() + BlockSizeX - 1) / BlockSizeX;
+
+                // TODO (Andrew) Actually get the offset, may not be possible
+                return FAVLayout(NumBlocksX * BlockBytes, 0, Size);
+            }
+            
+            break;
 #endif
 	default:
 		break;
@@ -112,15 +134,21 @@ TSharedPtr<FVideoResourceRHI> FVideoResourceRHI::Create(TSharedPtr<FAVDevice> co
 		FRHITextureCreateDesc TextureDesc = FRHITextureCreateDesc::Create2D(TEXT("AVCodecs Resource"), Descriptor.Width, Descriptor.Height, static_cast<EPixelFormat>(Descriptor.Format));
 
 		TextureDesc.SetClearValue(FClearValueBinding::None);
-		TextureDesc.SetFlags(ETextureCreateFlags::RenderTargetable);
+#if AVCODECS_USE_METAL
+        TextureDesc.SetFlags(ETextureCreateFlags::CPUReadback);
+		TextureDesc.SetInitialState(ERHIAccess::CPURead);
+        TextureDesc.DetermineInititialState();
+#else
+    	TextureDesc.SetFlags(ETextureCreateFlags::RenderTargetable);
 		TextureDesc.SetInitialState(ERHIAccess::Present);
+#endif
 		TextureDesc.SetNumMips(1);
 
 		if (RHIGetInterfaceType() == ERHIInterfaceType::Vulkan)
 		{
 			TextureDesc.AddFlags(ETextureCreateFlags::External);
 		}
-		else if (RHIGetInterfaceType() == ERHIInterfaceType::D3D11 || RHIGetInterfaceType() == ERHIInterfaceType::D3D12)
+		else
 		{
 			TextureDesc.AddFlags(ETextureCreateFlags::Shared);
 		}
@@ -157,12 +185,16 @@ TSharedPtr<FVideoResourceRHI> FVideoResourceRHI::Create(TSharedPtr<FAVDevice> co
 			TextureDesc.Extent.X *= 3;
 			Descriptor.RawDescriptor = new FVideoDescriptor(EVideoFormat::G16, TextureDesc.Extent.X, TextureDesc.Extent.Y);
 			break;
+        case EVideoFormat::BGRA:
+			TextureDesc.Format = EPixelFormat::PF_B8G8R8A8;
+			Descriptor.RawDescriptor = new FVideoDescriptor(EVideoFormat::BGRA, TextureDesc.Extent.X, TextureDesc.Extent.Y);
+			break;
 		default:
 			break;
 		}
 
 //TODO-TE THIS IS THE REAL DEAL?
-		return MakeShareable(new FVideoResourceRHI(Device.ToSharedRef(), { RHICreateTexture(TextureDesc), nullptr, 0 }, Descriptor));
+		return MakeShareable(new FVideoResourceRHI(Device.ToSharedRef(), { GDynamicRHI->RHICreateTexture(FRHICommandListExecutor::GetImmediateCommandList(), TextureDesc), nullptr, 0 }, Descriptor));
 	}
 
 	return nullptr;
@@ -172,7 +204,7 @@ FVideoResourceRHI::FVideoResourceRHI(TSharedRef<FAVDevice> const& Device, FRawDa
 	: TVideoResource(Device, GetLayoutFrom(Device, Raw.Texture), OverrideDescriptor)
 	, Raw(Raw)
 {
-#if PLATFORM_WINDOWS
+#if AVCODECS_USE_D3D
 	if (RHIGetInterfaceType() == ERHIInterfaceType::D3D12)
 	{
 		// Create a fence if we did not pass one in
@@ -431,7 +463,7 @@ DLLEXPORT FAVResult FAVExtension::TransformResource(TSharedPtr<FVideoResourceRHI
 	return FAVResult(EAVResult::ErrorMapping, TEXT("Input resource is not valid"), TEXT("RHI"));
 }
 
-#if PLATFORM_WINDOWS
+#if AVCODECS_USE_D3D
 
 template <>
 DLLEXPORT FAVResult FAVExtension::TransformResource(TSharedPtr<FVideoResourceD3D11>& OutResource, TSharedPtr<FVideoResourceRHI> const& InResource)
@@ -494,6 +526,7 @@ DLLEXPORT FAVResult FAVExtension::TransformResource(TSharedPtr<FVideoResourceD3D
 
 #endif
 
+#if AVCODECS_USE_VULKAN
 template <>
 DLLEXPORT FAVResult FAVExtension::TransformResource(TSharedPtr<FVideoResourceVulkan>& OutResource, TSharedPtr<FVideoResourceRHI> const& InResource)
 {
@@ -515,3 +548,27 @@ DLLEXPORT FAVResult FAVExtension::TransformResource(TSharedPtr<FVideoResourceVul
 
 	return FAVResult(EAVResult::ErrorMapping, TEXT("Input resource is not valid"), TEXT("RHI"));
 }
+#endif
+
+#if AVCODECS_USE_METAL
+template <>
+DLLEXPORT FAVResult FAVExtension::TransformResource(TSharedPtr<FVideoResourceMetal>& OutResource, TSharedPtr<FVideoResourceRHI> const& InResource)
+{
+	if (InResource.IsValid())
+	{
+		if (InResource->GetDevice()->HasContext<FVideoContextMetal>())
+		{		
+			OutResource = MakeShared<FVideoResourceMetal>(
+				InResource->GetDevice(),
+				static_cast<mtlpp::Texture*>(InResource->GetRaw().Texture->GetNativeResource()),
+				InResource->GetLayout());
+
+			return OutResource->Validate();
+		}
+
+		return FAVResult(EAVResult::ErrorMapping, TEXT("No Metal context found"), TEXT("RHI"));
+	}
+
+	return FAVResult(EAVResult::ErrorMapping, TEXT("Input resource is not valid"), TEXT("RHI"));
+}
+#endif
diff --git a/Engine/Plugins/Experimental/AVCodecs/AVCodecsCore/Source/AVCodecsCoreRHI/Public/Video/Decoders/VideoDecoderRHI.h b/Engine/Plugins/Experimental/AVCodecs/AVCodecsCore/Source/AVCodecsCoreRHI/Public/Video/Decoders/VideoDecoderRHI.h
--- a/Engine/Plugins/Experimental/AVCodecs/AVCodecsCore/Source/AVCodecsCoreRHI/Public/Video/Decoders/VideoDecoderRHI.h
+++ b/Engine/Plugins/Experimental/AVCodecs/AVCodecsCore/Source/AVCodecsCoreRHI/Public/Video/Decoders/VideoDecoderRHI.h
@@ -4,10 +4,16 @@
 
 #include "Video/VideoDecoder.h"
 #include "Video/Resources/VideoResourceRHI.h"
-#include "Video/Resources/VideoResourceVulkan.h"
 
-#if PLATFORM_WINDOWS
-#include "Video/Resources/Windows/VideoResourceD3D.h"
+#if AVCODECS_USE_D3D
+#include "Video/Resources/D3D/VideoResourceD3D.h"
+#include "Video/Resources/Vulkan/VideoResourceVulkan.h"
+#endif
+#if AVCODECS_USE_VULKAN
+#include "Video/Resources/Vulkan/VideoResourceVulkan.h"
+#endif
+#if AVCODECS_USE_METAL
+#include "Video/Resources/Metal/VideoResourceMetal.h"
 #endif
 
 template <typename TConfig>
@@ -44,11 +50,13 @@ public:
 
 		switch (GDynamicRHI->GetInterfaceType())
 		{
+#if AVCODECS_USE_VULKAN
 		case ERHIInterfaceType::Vulkan:
 			Child = FVideoDecoder::Wrap<FVideoResourceRHI, TConfig>(FVideoDecoder::Create<FVideoResourceVulkan, TConfig>(NewDevice, NewInstance));
 		
 			break;
-#if PLATFORM_WINDOWS
+#endif
+#if AVCODECS_USE_D3D
 		case ERHIInterfaceType::D3D11:
 			Child = FVideoDecoder::Wrap<FVideoResourceRHI, TConfig>(FVideoDecoder::Create<FVideoResourceD3D11, TConfig>(NewDevice, NewInstance));
 		
@@ -58,6 +66,12 @@ public:
 		
 			break;
 #endif
+#if AVCODECS_USE_METAL
+		case ERHIInterfaceType::Metal:
+			Child = FVideoDecoder::Wrap<FVideoResourceRHI, TConfig>(FVideoDecoder::Create<FVideoResourceMetal, TConfig>(NewDevice, NewInstance));
+
+			break;
+#endif
 		default:
 			return FAVResult(EAVResult::ErrorCreating, FString::Printf(TEXT("RHI type %d not supported"), GDynamicRHI->GetInterfaceType()), TEXT("RHI"));
 		}
diff --git a/Engine/Plugins/Experimental/AVCodecs/AVCodecsCore/Source/AVCodecsCoreRHI/Public/Video/Encoders/VideoEncoderRHI.h b/Engine/Plugins/Experimental/AVCodecs/AVCodecsCore/Source/AVCodecsCoreRHI/Public/Video/Encoders/VideoEncoderRHI.h
--- a/Engine/Plugins/Experimental/AVCodecs/AVCodecsCore/Source/AVCodecsCoreRHI/Public/Video/Encoders/VideoEncoderRHI.h
+++ b/Engine/Plugins/Experimental/AVCodecs/AVCodecsCore/Source/AVCodecsCoreRHI/Public/Video/Encoders/VideoEncoderRHI.h
@@ -4,10 +4,15 @@
 
 #include "Video/VideoEncoder.h"
 #include "Video/Resources/VideoResourceRHI.h"
-#include "Video/Resources/VideoResourceVulkan.h"
 
-#if PLATFORM_WINDOWS
-#include "Video/Resources/Windows/VideoResourceD3D.h"
+#if AVCODECS_USE_D3D
+#include "Video/Resources/D3D/VideoResourceD3D.h"
+#endif
+#if AVCODECS_USE_VULKAN
+#include "Video/Resources/Vulkan/VideoResourceVulkan.h"
+#endif
+#if AVCODECS_USE_METAL
+#include "Video/Resources/Metal/VideoResourceMetal.h"
 #endif
 
 template <typename TConfig>
@@ -44,11 +49,13 @@ public:
 
 		switch (GDynamicRHI->GetInterfaceType())
 		{
+#if AVCODECS_USE_VULKAN
 		case ERHIInterfaceType::Vulkan:
 			Child = FVideoEncoder::Wrap<FVideoResourceRHI, TConfig>(FVideoEncoder::Create<FVideoResourceVulkan, TConfig>(NewDevice, NewInstance));
 		
 			break;
-#if PLATFORM_WINDOWS
+#endif
+#if AVCODECS_USE_D3D
 		case ERHIInterfaceType::D3D11:
 			Child = FVideoEncoder::Wrap<FVideoResourceRHI, TConfig>(FVideoEncoder::Create<FVideoResourceD3D11, TConfig>(NewDevice, NewInstance));
 		
@@ -58,6 +65,12 @@ public:
 		
 			break;
 #endif
+#if AVCODECS_USE_METAL
+        case ERHIInterfaceType::Metal:
+			Child = FVideoEncoder::Wrap<FVideoResourceRHI, TConfig>(FVideoEncoder::Create<FVideoResourceMetal, TConfig>(NewDevice, NewInstance));
+		
+			break;
+#endif
 		default:
 			return FAVResult(EAVResult::ErrorCreating, FString::Printf(TEXT("RHI type %d not supported"), GDynamicRHI->GetInterfaceType()), TEXT("RHI"));
 		}
diff --git a/Engine/Plugins/Experimental/AVCodecs/AVCodecsCore/Source/AVCodecsCoreRHI/Public/Video/Resources/VideoResourceRHI.h b/Engine/Plugins/Experimental/AVCodecs/AVCodecsCore/Source/AVCodecsCoreRHI/Public/Video/Resources/VideoResourceRHI.h
--- a/Engine/Plugins/Experimental/AVCodecs/AVCodecsCore/Source/AVCodecsCoreRHI/Public/Video/Resources/VideoResourceRHI.h
+++ b/Engine/Plugins/Experimental/AVCodecs/AVCodecsCore/Source/AVCodecsCoreRHI/Public/Video/Resources/VideoResourceRHI.h
@@ -7,7 +7,7 @@
 #include "Video/VideoResource.h"
 #include "RHI.h"
 
-#if PLATFORM_WINDOWS
+#if AVCODECS_USE_D3D
 THIRD_PARTY_INCLUDES_START
 #include "Windows/WindowsHWrapper.h"
 #include "d3d12.h"
@@ -24,7 +24,7 @@ public:
 	struct FRawData
 	{
 		FTextureRHIRef Texture;
-#if PLATFORM_WINDOWS
+#if AVCODECS_USE_D3D
 		TRefCountPtr<ID3D12Fence> Fence; // TODO replace with an RHI fence and convert in a transform function
 #else 
 		FGPUFenceRHIRef Fence;
@@ -81,7 +81,7 @@ FAVResult FAVExtension::TransformResource(TSharedPtr<FVideoResourceRHI>& OutReso
 template <>
 FAVResult FAVExtension::TransformResource(TSharedPtr<FVideoResourceRHI>& OutResource, TSharedPtr<class FVideoResourceCPU> const& InResource);
 
-#if PLATFORM_WINDOWS
+#if AVCODECS_USE_D3D
 
 template <>
 FAVResult FAVExtension::TransformResource(TSharedPtr<class FVideoResourceD3D11>& OutResource, TSharedPtr<FVideoResourceRHI> const& InResource);
@@ -91,8 +91,19 @@ FAVResult FAVExtension::TransformResource(TSharedPtr<class FVideoResourceD3D12>&
 
 #endif
 
+#if AVCODECS_USE_VULKAN
+
 template <>
 FAVResult FAVExtension::TransformResource(TSharedPtr<class FVideoResourceVulkan>& OutResource, TSharedPtr<FVideoResourceRHI> const& InResource);
 
+#endif
+
+#if AVCODECS_USE_METAL
+
+template <>
+FAVResult FAVExtension::TransformResource(TSharedPtr<class FVideoResourceMetal>& OutResource, TSharedPtr<FVideoResourceRHI> const& InResource);
+
+#endif
+
 DECLARE_TYPEID(FVideoContextRHI, AVCODECSCORERHI_API);
 DECLARE_TYPEID(FVideoResourceRHI, AVCODECSCORERHI_API);
\ No newline at end of file
diff --git a/Engine/Plugins/Experimental/AVCodecs/NVCodecs/Source/NVCodecs/Private/Video/Resources/VideoResourceCUDA.cpp b/Engine/Plugins/Experimental/AVCodecs/NVCodecs/Source/NVCodecs/Private/Video/Resources/VideoResourceCUDA.cpp
--- a/Engine/Plugins/Experimental/AVCodecs/NVCodecs/Source/NVCodecs/Private/Video/Resources/VideoResourceCUDA.cpp
+++ b/Engine/Plugins/Experimental/AVCodecs/NVCodecs/Source/NVCodecs/Private/Video/Resources/VideoResourceCUDA.cpp
@@ -2,10 +2,10 @@
 
 #include "Video/Resources/VideoResourceCUDA.h"
 
-#include "Video/Resources/VideoResourceVulkan.h"
+#include "Video/Resources/Vulkan/VideoResourceVulkan.h"
 
 #if PLATFORM_WINDOWS
-	#include "Video/Resources/Windows/VideoResourceD3D.h"
+#include "Video/Resources/D3D/VideoResourceD3D.h"
 #endif
 
 REGISTER_TYPEID(FVideoContextCUDA);
@@ -15,13 +15,13 @@ TAVResult<CUarray_format> ConvertFormat(EVideoFormat Format)
 {
 	switch (Format)
 	{
-		case EVideoFormat::R8:
-		case EVideoFormat::BGRA:
-			return CU_AD_FORMAT_UNSIGNED_INT8;
-		case EVideoFormat::G16:
-			return CU_AD_FORMAT_UNSIGNED_INT16;
-		default:
-			return FAVResult(EAVResult::ErrorUnsupported, FString::Printf(TEXT("EVideoFormat format %d is not supported"), Format), TEXT("CUDA"));
+	case EVideoFormat::R8:
+	case EVideoFormat::BGRA:
+		return CU_AD_FORMAT_UNSIGNED_INT8;
+	case EVideoFormat::G16:
+		return CU_AD_FORMAT_UNSIGNED_INT16;
+	default:
+		return FAVResult(EAVResult::ErrorUnsupported, FString::Printf(TEXT("EVideoFormat format %d is not supported"), Format), TEXT("CUDA"));
 	}
 }
 
@@ -134,15 +134,6 @@ FVideoResourceCUDA::~FVideoResourceCUDA()
 			FAVResult::Log(EAVResult::ErrorUnmapping, TEXT("Failed to clean up external memory"), TEXT("CUDA"), Result);
 		}
 	}
-
-	if (ExternalSemaphore)
-	{
-		CUresult const Result = FCUDAModule::CUDA().cuDestroyExternalSemaphore(ExternalSemaphore);
-		if (Result != CUDA_SUCCESS)
-		{
-			FAVResult::Log(EAVResult::ErrorUnmapping, TEXT("Failed to destroy external semaphore"), TEXT("CUDA"), Result);
-		}
-	}
 }
 
 FAVResult FVideoResourceCUDA::Validate() const
@@ -222,7 +213,7 @@ FAVResult FVideoResourceCUDA::CopyFromAsync(CUdeviceptr Source, uint32 MapPitch)
 		CUDA_EXTERNAL_SEMAPHORE_SIGNAL_PARAMS params = {};
 		params.params.fence.value = SemaphoreValue + 1;
 
-		FCUDAModule::CUDA().cuSignalExternalSemaphoresAsync(&ExternalSemaphore, &params, 1, cuStream);
+		FCUDAModule::CUDA().cuSignalExternalSemaphoresAsync(&ExternalSemaphore, &params, 1, cuStream);		
 	}
 
 	// Destroy the stream this will happen only when all tasks on the stream have completed
@@ -278,7 +269,7 @@ FAVResult FVideoResourceCUDA::CopyToAsync(CUdeviceptr Target, uint32 MapPitch)
 		CUDA_EXTERNAL_SEMAPHORE_SIGNAL_PARAMS params = {};
 		params.params.fence.value = SemaphoreValue + 1;
 
-		FCUDAModule::CUDA().cuSignalExternalSemaphoresAsync(&ExternalSemaphore, &params, 1, cuStream);
+		FCUDAModule::CUDA().cuSignalExternalSemaphoresAsync(&ExternalSemaphore, &params, 1, cuStream);		
 	}
 
 	// Destroy the stream this will happen only when all tasks on the stream have completed
@@ -330,7 +321,7 @@ DLLEXPORT FAVResult FAVExtension::TransformResource(TSharedPtr<FVideoResourceCUD
 		if (InResource->GetDevice()->HasContext<FVideoContextCUDA>())
 		{
 			// Depending on how we are running the engine the texture we are looking to bind could be
-			// stored as part of a heap or as a commited resource
+			// stored as part of a heap or as a commited resource 
 			CUDA_EXTERNAL_MEMORY_HANDLE_DESC ExternalResouceDesc = {};
 			if (InResource->GetHeap().IsValid())
 			{
@@ -355,7 +346,7 @@ DLLEXPORT FAVResult FAVExtension::TransformResource(TSharedPtr<FVideoResourceCUD
 			ExternalFenceDesc.handle.win32.handle = InResource->GetFenceSharedHandle();
 
 			OutResource = MakeShared<FVideoResourceCUDA>(InResource->GetDevice(), ExternalResouceDesc, ExternalFenceDesc, InResource.ToSharedRef(), InResource->GetFenceValue());
-
+			
 			return OutResource->Validate();
 		}
 
diff --git a/Engine/Plugins/Experimental/AVCodecs/NVCodecs/Source/NVDEC/Private/NVDECModule.cpp b/Engine/Plugins/Experimental/AVCodecs/NVCodecs/Source/NVDEC/Private/NVDECModule.cpp
--- a/Engine/Plugins/Experimental/AVCodecs/NVCodecs/Source/NVDEC/Private/NVDECModule.cpp
+++ b/Engine/Plugins/Experimental/AVCodecs/NVCodecs/Source/NVDEC/Private/NVDECModule.cpp
@@ -3,10 +3,10 @@
 #include "Video/Decoders/VideoDecoderNVDEC.h"
 #include "Video/Decoders/Configs/VideoDecoderConfigH264.h"
 #include "Video/Decoders/Configs/VideoDecoderConfigH265.h"
-#include "Video/Resources/VideoResourceVulkan.h"
+#include "Video/Resources/Vulkan/VideoResourceVulkan.h"
 
 #if PLATFORM_WINDOWS
-#include "Video/Resources/Windows/VideoResourceD3D.h"
+#include "Video/Resources/D3D/VideoResourceD3D.h"
 #endif
 
 class FNVDECModule : public IModuleInterface
@@ -14,6 +14,7 @@ class FNVDECModule : public IModuleInterface
 public:
 	virtual void StartupModule() override
 	{
+		// TODO (william.belcher): NVDEC decoding doesn't currently work
 		FVideoDecoder
 			::RegisterPermutationsOf<FVideoDecoderNVDEC>
 			::With<
