Fix: Remove memory leak in TVideoEncoder

#11068

diff --git a/Engine/Plugins/Experimental/AVCodecs/AVCodecsCore/Source/AVCodecsCore/Public/Video/VideoEncoder.h b/Engine/Plugins/Experimental/AVCodecs/AVCodecsCore/Source/AVCodecsCore/Public/Video/VideoEncoder.h
--- a/Engine/Plugins/Experimental/AVCodecs/AVCodecsCore/Source/AVCodecsCore/Public/Video/VideoEncoder.h
+++ b/Engine/Plugins/Experimental/AVCodecs/AVCodecsCore/Source/AVCodecsCore/Public/Video/VideoEncoder.h
@@ -151,7 +151,7 @@ public:
 
 			FScopeLock const Lock = Resource->LockScope();
 
-			TSharedPtr<TChildResource>& MappedResource = Resource->template PinMapping<TChildResource>();
+			TSharedPtr<TChildResource> MappedResource = Resource->template PinMapping<TChildResource>();
 			if (!MappedResource.IsValid() || MappedResource->Validate().IsNotSuccess())
 			{
 				Result = FAVExtension::TransformResource<TChildResource, TResource>(MappedResource, Resource);
