QOL: Remove requirement for textures passed to VTCodecs to be created with the CPU_READBACK flag

diff --git a/Engine/Plugins/Experimental/AVCodecs/AVCodecsCore/Source/AVCodecsCore/Private/Video/Resources/Metal/VideoResourceMetal.cpp b/Engine/Plugins/Experimental/AVCodecs/AVCodecsCore/Source/AVCodecsCore/Private/Video/Resources/Metal/VideoResourceMetal.cpp
--- a/Engine/Plugins/Experimental/AVCodecs/AVCodecsCore/Source/AVCodecsCore/Private/Video/Resources/Metal/VideoResourceMetal.cpp
+++ b/Engine/Plugins/Experimental/AVCodecs/AVCodecsCore/Source/AVCodecsCore/Private/Video/Resources/Metal/VideoResourceMetal.cpp
@@ -7,18 +7,14 @@
 REGISTER_TYPEID(FVideoContextMetal);
 REGISTER_TYPEID(FVideoResourceMetal);
 
-static TAVResult<EVideoFormat> ConvertFormat(MTL::PixelFormat Format)
+static TAVResult<EVideoFormat> ConvertFormat(OSType Format)
 {
 	switch (Format)
 	{
-    case MTL::PixelFormatBGRA8Unorm:
-    case MTL::PixelFormatBGRA8Unorm_sRGB:
+        case kCVPixelFormatType_32BGRA:
 		return EVideoFormat::BGRA;
-	case MTL::PixelFormatRGB10A2Unorm:
+	case kCVPixelFormatType_ARGB2101010LEPacked:
 		return EVideoFormat::ABGR10;
-	case MTL::PixelFormatR8Unorm:
-	case MTL::PixelFormatR8Uint:
-		return EVideoFormat::R8;
 	default:
 		return FAVResult(EAVResult::ErrorUnsupported, FString::Printf(TEXT("MTL::PixelFormat format %d is not supported"), Format), TEXT("Metal"));
 	}
@@ -29,19 +25,32 @@ FVideoContextMetal::FVideoContextMetal(MTL::Device* Device)
 {
 }
 
-FVideoDescriptor FVideoResourceMetal::GetDescriptorFrom(TSharedRef<FAVDevice> const& Device, MTL::Texture* Raw)
+FVideoDescriptor FVideoResourceMetal::GetDescriptorFrom(TSharedRef<FAVDevice> const& Device, CVPixelBufferRef Raw)
 {
-    uint32_t Width = Raw->width();
-    uint32_t Height = Raw->height();
-    TAVResult<EVideoFormat> ConvertedFormat = ConvertFormat(Raw->pixelFormat());
+    uint32_t Width = CVPixelBufferGetWidth(Raw);
+    uint32_t Height = CVPixelBufferGetHeight(Raw);
+    TAVResult<EVideoFormat> ConvertedFormat = ConvertFormat(CVPixelBufferGetPixelFormatType(Raw));
     
 	return FVideoDescriptor(ConvertedFormat, Width, Height);
 }
 
-FVideoResourceMetal::FVideoResourceMetal(TSharedRef<FAVDevice> const& Device, MTL::Texture* Raw, FAVLayout const& Layout)
+FVideoResourceMetal::FVideoResourceMetal(TSharedRef<FAVDevice> const& Device, CVPixelBufferRef Raw, FAVLayout const& Layout)
 	: TVideoResource(Device, Layout, GetDescriptorFrom(Device, Raw))
 	, Raw(Raw)
 {
+    if (Raw)
+    {
+        CFRetain(Raw);
+    }
+}
+
+FVideoResourceMetal::~FVideoResourceMetal()
+{
+    if (Raw)
+    {
+        CFRelease(Raw);
+        Raw = nullptr;
+    }
 }
 
 FAVResult FVideoResourceMetal::Validate() const
diff --git a/Engine/Plugins/Experimental/AVCodecs/AVCodecsCore/Source/AVCodecsCore/Public/Video/Resources/Metal/VideoResourceMetal.h b/Engine/Plugins/Experimental/AVCodecs/AVCodecsCore/Source/AVCodecsCore/Public/Video/Resources/Metal/VideoResourceMetal.h
--- a/Engine/Plugins/Experimental/AVCodecs/AVCodecsCore/Source/AVCodecsCore/Public/Video/Resources/Metal/VideoResourceMetal.h
+++ b/Engine/Plugins/Experimental/AVCodecs/AVCodecsCore/Source/AVCodecsCore/Public/Video/Resources/Metal/VideoResourceMetal.h
@@ -30,15 +30,16 @@ public:
 class AVCODECSCORE_API FVideoResourceMetal : public TVideoResource<FVideoContextMetal>
 {
 private:
-	MTL::Texture* Raw;
+    CVPixelBufferRef Raw;
 
 public:
-	static FVideoDescriptor GetDescriptorFrom(TSharedRef<FAVDevice> const& Device, MTL::Texture* Raw);
+	static FVideoDescriptor GetDescriptorFrom(TSharedRef<FAVDevice> const& Device, CVPixelBufferRef Raw);
 
-	FORCEINLINE MTL::Texture* GetRaw() const { return Raw; }
-
-	FVideoResourceMetal(TSharedRef<FAVDevice> const& Device, MTL::Texture* Raw, FAVLayout const& Layout);
+	FORCEINLINE CVPixelBufferRef GetRaw() const { return Raw; }
 
+	FVideoResourceMetal(TSharedRef<FAVDevice> const& Device, CVPixelBufferRef Raw, FAVLayout const& Layout);
+    virtual ~FVideoResourceMetal() override;
+    
 	virtual FAVResult Validate() const override;
 };
 
diff --git a/Engine/Plugins/Experimental/AVCodecs/AVCodecsCore/Source/AVCodecsCoreRHI/Private/Video/Resources/VideoResourceRHI.cpp b/Engine/Plugins/Experimental/AVCodecs/AVCodecsCore/Source/AVCodecsCoreRHI/Private/Video/Resources/VideoResourceRHI.cpp
--- a/Engine/Plugins/Experimental/AVCodecs/AVCodecsCore/Source/AVCodecsCoreRHI/Private/Video/Resources/VideoResourceRHI.cpp
+++ b/Engine/Plugins/Experimental/AVCodecs/AVCodecsCore/Source/AVCodecsCoreRHI/Private/Video/Resources/VideoResourceRHI.cpp
@@ -21,6 +21,7 @@
 #endif
 
 #include "Async/Async.h"
+#include "CoreGlobals.h"
 #include "HAL/Event.h"
 #include "RHIStaticStates.h"
 #include "MediaShaders.h"
@@ -31,6 +32,13 @@
 REGISTER_TYPEID(FVideoContextRHI);
 REGISTER_TYPEID(FVideoResourceRHI);
 
+#define CFSafeRelease(x)                \
+    if (x != nullptr)                   \
+    {                                   \
+        CFRelease(x);                   \
+        x = nullptr;                    \
+    }
+
 FAVLayout FVideoResourceRHI::GetLayoutFrom(TSharedRef<FAVDevice> const& Device, FTextureRHIRef const& Raw)
 {
 	switch (GDynamicRHI->GetInterfaceType())
@@ -134,14 +142,9 @@ TSharedPtr<FVideoResourceRHI> FVideoResourceRHI::Create(TSharedPtr<FAVDevice> co
 		FRHITextureCreateDesc TextureDesc = FRHITextureCreateDesc::Create2D(TEXT("AVCodecs Resource"), Descriptor.Width, Descriptor.Height, static_cast<EPixelFormat>(Descriptor.Format));
 
 		TextureDesc.SetClearValue(FClearValueBinding::None);
-#if AVCODECS_USE_METAL
-        TextureDesc.SetFlags(ETextureCreateFlags::CPUReadback);
-		TextureDesc.SetInitialState(ERHIAccess::CPURead);
-        TextureDesc.DetermineInititialState();
-#else
     	TextureDesc.SetFlags(ETextureCreateFlags::RenderTargetable);
 		TextureDesc.SetInitialState(ERHIAccess::Present);
-#endif
+
 		TextureDesc.SetNumMips(1);
 
 		if (RHIGetInterfaceType() == ERHIInterfaceType::Vulkan)
@@ -558,12 +561,83 @@ DLLEXPORT FAVResult FAVExtension::TransformResource(TSharedPtr<FVideoResourceMet
 	{
 		if (InResource->GetDevice()->HasContext<FVideoContextMetal>())
 		{		
-			OutResource = MakeShared<FVideoResourceMetal>(
-				InResource->GetDevice(),
-				static_cast<MTL::Texture*>(InResource->GetRaw().Texture->GetNativeResource()),
-				InResource->GetLayout());
+            static CVMetalTextureCacheRef TextureCache = nullptr;
+            
+            if (TextureCache == nullptr)
+            {
+                id<MTLDevice> Device = (__bridge id<MTLDevice>)GDynamicRHI->RHIGetNativeDevice();
+                check(Device);
+                
+                CVReturn Result = CVMetalTextureCacheCreate(kCFAllocatorDefault, nullptr, Device, nullptr, &TextureCache);
+                if (Result != kCVReturnSuccess)
+                {
+                    return FAVResult(EAVResult::Error, TEXT("Failed to create CVMetalTextureCacheRef"), TEXT("RHI"), Result);
+                }
+            }
+            
+            const FVideoDescriptor& Descriptor = InResource->GetDescriptor();
+                        
+            CFMutableDictionaryRef SourceAttributes = CFDictionaryCreateMutable(kCFAllocatorDefault, 1, &kCFTypeDictionaryKeyCallBacks, &kCFTypeDictionaryValueCallBacks);
+            CFDictionarySetValue(SourceAttributes, kCVPixelBufferOpenGLCompatibilityKey, kCFBooleanTrue);
+            CFDictionaryRef IOSurfaceValue = CFDictionaryCreate(kCFAllocatorDefault, nullptr, nullptr, 0, &kCFTypeDictionaryKeyCallBacks, &kCFTypeDictionaryValueCallBacks);
+            CFDictionarySetValue(SourceAttributes, kCVPixelBufferIOSurfacePropertiesKey, IOSurfaceValue);
+            int64 PixelType = kCVPixelFormatType_32BGRA;
+            CFNumberRef PixelFormat = CFNumberCreate(nullptr, kCFNumberLongType, &PixelType);
+            CFDictionarySetValue(SourceAttributes, kCVPixelBufferPixelFormatTypeKey, PixelFormat);
+
+            CFSafeRelease(IOSurfaceValue);
+            CFSafeRelease(PixelFormat);
+
+            CVPixelBufferRef PixelBuffer;
+            CVReturn Result = CVPixelBufferCreate(kCFAllocatorDefault, Descriptor.Width, Descriptor.Height, kCVPixelFormatType_32BGRA, SourceAttributes, &PixelBuffer);
+            if (Result != kCVReturnSuccess)
+            {
+                return FAVResult(EAVResult::Error, TEXT("Failed to create CVPixelBufferRef"), TEXT("RHI"), Result);
+            }
+            CFSafeRelease(SourceAttributes);
 
-			return OutResource->Validate();
+            CVMetalTextureRef TextureRef;
+            Result = CVMetalTextureCacheCreateTextureFromImage(kCFAllocatorDefault, TextureCache, PixelBuffer, nullptr, MTLPixelFormatBGRA8Unorm_sRGB, Descriptor.Width, Descriptor.Height, 0, &TextureRef);
+            if (Result != kCVReturnSuccess)
+            {
+                return FAVResult(EAVResult::Error, TEXT("Failed to create CVMetalTextureRef"), TEXT("RHI"), Result);
+            }
+
+            FGPUFenceRHIRef Fence;
+            FTextureRHIRef Destination;
+            const FRHITextureCreateDesc Desc =
+                FRHITextureCreateDesc::Create2D(TEXT("FAVExtension::TransformResource"), Descriptor.Width, Descriptor.Height, PF_B8G8R8A8)
+                .SetFlags(ETextureCreateFlags::SRGB | ETextureCreateFlags::Dynamic | ETextureCreateFlags::NoTiling | ETextureCreateFlags::ShaderResource)
+                .SetBulkData(new FCVBulkData(TextureRef));
+                        
+            ENQUEUE_RENDER_COMMAND(FAVExtensionTransformResource)(
+                [&Source = InResource->GetRaw().Texture, &Desc, &Fence, &Destination](FRHICommandListImmediate& RHICmdList)
+            {
+                Fence = RHICreateGPUFence(TEXT("FAVExtension::TransformResource"));
+                Destination = RHICreateTexture(Desc);
+                RHICmdList.CopyTexture(Source, Destination, FRHICopyTextureInfo());
+                            
+                RHICmdList.WriteGPUFence(Fence);
+            });
+            
+            while(!Fence || !Fence->Poll())
+            {
+				if(IsEngineExitRequested())
+				{
+					return FAVResult(EAVResult::Error, TEXT("Engine exit requested before texture copy was complete"), TEXT("RHI"));
+				}
+
+                FPlatformProcess::YieldThread();
+            }
+
+            OutResource = MakeShared<FVideoResourceMetal>(InResource->GetDevice(),
+                                                          PixelBuffer,
+                                                          InResource->GetLayout());
+
+            CFSafeRelease(TextureRef);
+            CFSafeRelease(PixelBuffer);
+            
+            return OutResource->Validate();
 		}
 
 		return FAVResult(EAVResult::ErrorMapping, TEXT("No Metal context found"), TEXT("RHI"));
diff --git a/Engine/Plugins/Experimental/AVCodecs/AVCodecsCore/Source/AVCodecsCoreRHI/Public/Video/Resources/VideoResourceRHI.h b/Engine/Plugins/Experimental/AVCodecs/AVCodecsCore/Source/AVCodecsCoreRHI/Public/Video/Resources/VideoResourceRHI.h
--- a/Engine/Plugins/Experimental/AVCodecs/AVCodecsCore/Source/AVCodecsCoreRHI/Public/Video/Resources/VideoResourceRHI.h
+++ b/Engine/Plugins/Experimental/AVCodecs/AVCodecsCore/Source/AVCodecsCoreRHI/Public/Video/Resources/VideoResourceRHI.h
@@ -72,6 +72,45 @@ protected:
 	virtual TSharedPtr<FVideoResourceRHI> TryResolve(TSharedPtr<FAVDevice> const& Device, FVideoDescriptor const& Descriptor) override;
 };
 
+#if AVCODECS_USE_METAL
+/**
+ * Passes a CVMetalTextureRef or CVPixelBufferRef through to the RHI to wrap in an RHI texture without traversing system memory.
+ */
+class AVCODECSCORERHI_API FCVBulkData : public FResourceBulkDataInterface
+{
+public:
+    FCVBulkData(CFTypeRef InImageBuffer)
+        : ImageBuffer(InImageBuffer)
+    {
+        check(ImageBuffer);
+        CFRetain(ImageBuffer);
+    }
+    virtual ~FCVBulkData()
+    {
+        CFRelease(ImageBuffer);
+        ImageBuffer = nullptr;
+    }
+public:
+    virtual void Discard() override
+    {
+        delete this;
+    }
+    virtual const void* GetResourceBulkData() const override
+    {
+        return ImageBuffer;
+    }
+    virtual uint32 GetResourceBulkDataSize() const override
+    {
+        return ImageBuffer ? ~0u : 0;
+    }
+    virtual EBulkDataType GetResourceType() const override
+    {
+        return EBulkDataType::MediaTexture;
+    }
+    CFTypeRef ImageBuffer;
+};
+#endif
+
 /*
 * Handles pixel format or colorspace transformations
 */
@@ -106,4 +145,4 @@ FAVResult FAVExtension::TransformResource(TSharedPtr<class FVideoResourceMetal>&
 #endif
 
 DECLARE_TYPEID(FVideoContextRHI, AVCODECSCORERHI_API);
-DECLARE_TYPEID(FVideoResourceRHI, AVCODECSCORERHI_API);
\ No newline at end of file
+DECLARE_TYPEID(FVideoResourceRHI, AVCODECSCORERHI_API);
diff --git a/Engine/Plugins/Experimental/AVCodecs/VTCodecs/Source/VTCodecs/Public/Video/Decoders/VideoDecoderVT.hpp b/Engine/Plugins/Experimental/AVCodecs/VTCodecs/Source/VTCodecs/Public/Video/Decoders/VideoDecoderVT.hpp
--- a/Engine/Plugins/Experimental/AVCodecs/VTCodecs/Source/VTCodecs/Public/Video/Decoders/VideoDecoderVT.hpp
+++ b/Engine/Plugins/Experimental/AVCodecs/VTCodecs/Source/VTCodecs/Public/Video/Decoders/VideoDecoderVT.hpp
@@ -299,10 +299,8 @@ FAVResult TVideoDecoderVT<TResource>::ReceiveFrame(TResolvableVideoResource<TRes
 
             void* PixelPtr = CVPixelBufferGetBaseAddress(Frame->ImageBuffer);
 
-            MTL::Texture* RawTexture = InOutResource->GetRaw();
-
             // Do copy into the VideoResource
-            InOutResource->GetRaw()->replaceRegion(MTL::Region(0, 0, Width, Height), 0, PixelPtr, Width * 4);
+//            InOutResource->GetRaw()->replaceRegion(MTL::Region(0, 0, Width, Height), 0, PixelPtr, Width * 4);
 
             CVPixelBufferUnlockBaseAddress(Frame->ImageBuffer, kCVPixelBufferLock_ReadOnly);
 
diff --git a/Engine/Plugins/Experimental/AVCodecs/VTCodecs/Source/VTCodecs/Public/Video/Encoders/VideoEncoderVT.hpp b/Engine/Plugins/Experimental/AVCodecs/VTCodecs/Source/VTCodecs/Public/Video/Encoders/VideoEncoderVT.hpp
--- a/Engine/Plugins/Experimental/AVCodecs/VTCodecs/Source/VTCodecs/Public/Video/Encoders/VideoEncoderVT.hpp
+++ b/Engine/Plugins/Experimental/AVCodecs/VTCodecs/Source/VTCodecs/Public/Video/Encoders/VideoEncoderVT.hpp
@@ -218,7 +218,6 @@ FAVResult TVideoEncoderVT<TResource>::SetEncoderBitrate(FVideoEncoderConfigVT co
 template <typename TResource>
 FAVResult TVideoEncoderVT<TResource>::SendFrame(TSharedPtr<FVideoResourceMetal> const& Resource, uint32 Timestamp, bool bForceKeyframe)
 {
-    
     if(IsOpen())
     {
         FVideoEncoderConfigVT& PendingConfig = this->EditPendingConfig();
@@ -232,34 +231,6 @@ FAVResult TVideoEncoderVT<TResource>::SendFrame(TSharedPtr<FVideoResourceMetal>
 
         if(Resource.IsValid())
         {
-            // Get a pixel buffer from the pool and copy frame data over.
-            CVPixelBufferPoolRef PixelBufferPool = VTCompressionSessionGetPixelBufferPool(Encoder);
-            if(!PixelBufferPool)
-            {
-                return FAVResult(EAVResult::Error, TEXT("Failed to get pixel buffer pool"), TEXT("VT"));
-            }
-
-            CVPixelBufferRef PixelBuffer;
-            CVReturn Result = CVPixelBufferPoolCreatePixelBuffer(nullptr, PixelBufferPool, &PixelBuffer);
-            if(Result != kCVReturnSuccess)
-            {
-                return FAVResult(EAVResult::Error, TEXT("Failed to create pixel buffer"), TEXT("VT"));
-            }
-
-            Result = CVPixelBufferLockBaseAddress(PixelBuffer, 0);
-            if (Result != kCVReturnSuccess) 
-            {
-                return FAVResult(EAVResult::Error, TEXT("Failed to lock base address"), TEXT("VT"));
-            }
-
-            // NOTE (belchy06): GetBytes assumes the raw texture has been created with with TexCreate_CPUReadback
-            MTL::Texture* RawTexture = Resource->GetRaw();
-            uint32_t Width = RawTexture->width();
-            uint32_t Height = RawTexture->height();
-
-            RawTexture->getBytes(reinterpret_cast<uint8*>(CVPixelBufferGetBaseAddressOfPlane(PixelBuffer, 0)), CVPixelBufferGetBytesPerRow(PixelBuffer), MTL::Region(0, 0, Width, Height), 0);
-            CVPixelBufferUnlockBaseAddress(PixelBuffer, 0);
-            
             CMTime PresentationTime = CMTimeMake(Timestamp, 1000);
             CFDictionaryRef FrameProperties = nullptr;
             if (bForceKeyframe) 
@@ -272,16 +243,9 @@ FAVResult TVideoEncoderVT<TResource>::SendFrame(TSharedPtr<FVideoResourceMetal>
             TUniquePtr<EncodeParams> Params = MakeUnique<EncodeParams>();
             Params.Reset(new EncodeParams(PendingConfig.Codec, PresentationTime));
 
-            OSStatus Status = VTCompressionSessionEncodeFrame(Encoder, PixelBuffer, PresentationTime, kCMTimeInvalid, FrameProperties, (void*)Params.Release(), nullptr);
+            OSStatus Status = VTCompressionSessionEncodeFrame(Encoder, Resource->GetRaw(), PresentationTime, kCMTimeInvalid, FrameProperties, (void*)Params.Release(), nullptr);
 
-            CONDITIONAL_RELEASE(FrameProperties);
-
-            if (PixelBuffer) 
-            {
-                CVPixelBufferRelease(PixelBuffer);
-                PixelBuffer = nullptr;
-            }
-        
+            CONDITIONAL_RELEASE(FrameProperties);        
 
             if(Status != 0)
             {
