Fix: Build health issues from CL30232152

diff --git a/Engine/Plugins/Experimental/AVCodecs/AVCodecsCore/Source/AVCodecsCoreRHI/Public/Video/Resources/VideoResourceRHI.h b/Engine/Plugins/Experimental/AVCodecs/AVCodecsCore/Source/AVCodecsCoreRHI/Public/Video/Resources/VideoResourceRHI.h
--- a/Engine/Plugins/Experimental/AVCodecs/AVCodecsCore/Source/AVCodecsCoreRHI/Public/Video/Resources/VideoResourceRHI.h
+++ b/Engine/Plugins/Experimental/AVCodecs/AVCodecsCore/Source/AVCodecsCoreRHI/Public/Video/Resources/VideoResourceRHI.h
@@ -5,6 +5,7 @@
 #include "AVContext.h"
 #include "AVExtension.h"
 #include "Video/VideoResource.h"
+#include "Containers/ResourceArray.h"
 #include "RHI.h"
 
 #if AVCODECS_USE_D3D
