QOL: Add VP9 profile enum

diff --git a/Engine/Plugins/Experimental/AVCodecs/AVCodecsCore/Source/AVCodecsCore/Private/Video/CodecUtils/CodecUtilsVP9.cpp b/Engine/Plugins/Experimental/AVCodecs/AVCodecsCore/Source/AVCodecsCore/Private/Video/CodecUtils/CodecUtilsVP9.cpp
--- a/Engine/Plugins/Experimental/AVCodecs/AVCodecsCore/Source/AVCodecsCore/Private/Video/CodecUtils/CodecUtilsVP9.cpp
+++ b/Engine/Plugins/Experimental/AVCodecs/AVCodecsCore/Source/AVCodecsCore/Private/Video/CodecUtils/CodecUtilsVP9.cpp
@@ -4,77 +4,77 @@
 
 namespace UE::AVCodecCore::VP9
 {
-    FAVResult ParseHeader(FBitstreamReader& Bitstream, Header_t& OutHeader)
-    {
-        Bitstream.Read(OutHeader.frame_marker); // u(2)
-        if (OutHeader.frame_marker != 2)
+	FAVResult ParseHeader(FBitstreamReader& Bitstream, Header_t& OutHeader)
+	{
+		Bitstream.Read(OutHeader.frame_marker); // u(2)
+		if (OutHeader.frame_marker != 2)
 		{
-            return FAVResult(EAVResult::Error, TEXT("frame_marker != 2"), TEXT("VP9"));
+			return FAVResult(EAVResult::Error, TEXT("frame_marker != 2"), TEXT("VP9"));
 		}
 
-        Bitstream.Read(OutHeader.profile_low_bit,   // u(1)
-                       OutHeader.profile_high_bit); // u(1)
-        OutHeader.Profile = (OutHeader.profile_high_bit << 1) + OutHeader.profile_low_bit;
+		Bitstream.Read(OutHeader.profile_low_bit, // u(1)
+			OutHeader.profile_high_bit);		  // u(1)
+		OutHeader.Profile = static_cast<EVP9Profile>((OutHeader.profile_high_bit << 1) + OutHeader.profile_low_bit);
 
-        if(OutHeader.Profile == 3)
-        {
-            if (Bitstream.ReadBits(1) != 0)
+		if (OutHeader.Profile == EVP9Profile::Profile3)
+		{
+			if (Bitstream.ReadBits(1) != 0)
 			{
-                // Reserved Zero
-                return FAVResult(EAVResult::Error, TEXT("Reserved zero bit was not zero!"), TEXT("VP9"));
+				// Reserved Zero
+				return FAVResult(EAVResult::Error, TEXT("Reserved zero bit was not zero!"), TEXT("VP9"));
 			}
-        }
-
-        Bitstream.Read(OutHeader.show_existing_frame); // u(1)
-        if(OutHeader.show_existing_frame == 1)
-        {
-            Bitstream.Read(OutHeader.frame_to_show_map_idx); // u(3);
-            return EAVResult::Success;
-        }
-
-        Bitstream.Read(OutHeader.frame_type,            // u(1)
-                       OutHeader.show_frame,            // u(1)
-                       OutHeader.error_resilient_mode); // u(1)
-
-        if (OutHeader.frame_type == 0)
-        {
-            // Is Keyframe
-            // frame sync code
-            if (ParseFrameSyncCode(Bitstream, OutHeader) != EAVResult::Success)
+		}
+
+		Bitstream.Read(OutHeader.show_existing_frame); // u(1)
+		if (OutHeader.show_existing_frame == 1)
+		{
+			Bitstream.Read(OutHeader.frame_to_show_map_idx); // u(3);
+			return EAVResult::Success;
+		}
+
+		Bitstream.Read(OutHeader.frame_type, // u(1)
+			OutHeader.show_frame,			 // u(1)
+			OutHeader.error_resilient_mode); // u(1)
+
+		if (OutHeader.frame_type == 0)
+		{
+			// Is Keyframe
+			// frame sync code
+			if (ParseFrameSyncCode(Bitstream, OutHeader) != EAVResult::Success)
 			{
-                return FAVResult(EAVResult::Error, TEXT("Failed to parse frame sync code"), TEXT("VP9"));
+				return FAVResult(EAVResult::Error, TEXT("Failed to parse frame sync code"), TEXT("VP9"));
 			}
 
 			// color config.
 			if (ParseColorConfig(Bitstream, OutHeader) != EAVResult::Success)
 			{
-                return FAVResult(EAVResult::Error, TEXT("Failed to parse color config"), TEXT("VP9"));
+				return FAVResult(EAVResult::Error, TEXT("Failed to parse color config"), TEXT("VP9"));
+			}
+
+			// frame size
+			if (ParseFrameSize(Bitstream, OutHeader) != EAVResult::Success)
+			{
+				return FAVResult(EAVResult::Error, TEXT("Failed to parse frame size"), TEXT("VP9"));
+			}
+		}
+		else
+		{
+			// Not keyframe
+			if (OutHeader.show_frame == 0)
+			{
+				Bitstream.Read(OutHeader.intra_only); // u(1)
 			}
 
-            // frame size
-            if(ParseFrameSize(Bitstream, OutHeader) != EAVResult::Success)
-            {
-                return FAVResult(EAVResult::Error, TEXT("Failed to parse frame size"), TEXT("VP9"));
-            }
-        }           
-        else
-        {
-            // Not keyframe
-            if(OutHeader.show_frame == 0)
-            {
-                Bitstream.Read(OutHeader.intra_only); // u(1)
-            }
-
-            if(OutHeader.error_resilient_mode == 0)
-            {
-                Bitstream.Read(OutHeader.reset_frame_context); // u(2)
-            }
+			if (OutHeader.error_resilient_mode == 0)
+			{
+				Bitstream.Read(OutHeader.reset_frame_context); // u(2)
+			}
 
 			if (OutHeader.intra_only == 1)
 			{
 				if (ParseFrameSyncCode(Bitstream, OutHeader) != EAVResult::Success)
 				{
-                    return FAVResult(EAVResult::Error, TEXT("Failed to parse frame sync code"), TEXT("VP9"));
+					return FAVResult(EAVResult::Error, TEXT("Failed to parse frame sync code"), TEXT("VP9"));
 				}
 
 				if (((OutHeader.profile_high_bit << 1) + OutHeader.profile_low_bit) > 0)
@@ -82,7 +82,7 @@ namespace UE::AVCodecCore::VP9
 					// Color config.
 					if (ParseColorConfig(Bitstream, OutHeader) != EAVResult::Success)
 					{
-                        return FAVResult(EAVResult::Error, TEXT("Failed to parse color config"), TEXT("VP9"));
+						return FAVResult(EAVResult::Error, TEXT("Failed to parse color config"), TEXT("VP9"));
 					}
 				}
 				else
@@ -90,44 +90,44 @@ namespace UE::AVCodecCore::VP9
 					OutHeader.color_space = EColorSpace::BT_601;
 					OutHeader.sub_sampling = ESubSampling::k420;
 				}
-			} 
-
-            // TODO (william.belcher): More VP9 "P" and "B" frame values. Not needed at the moment
-        }
-        
-        return EAVResult::Success;
-    }
-
-    FAVResult ParseFrameSyncCode(FBitstreamReader& Bitstream, Header_t& OutHeader)
-    {
-        Bitstream.Read(OutHeader.frame_sync_byte_0,  // u(8)
-                       OutHeader.frame_sync_byte_1,  // u(8)
-                       OutHeader.frame_sync_byte_2); // u(8)
-        return ((OutHeader.frame_sync_byte_0 == 0x49) && (OutHeader.frame_sync_byte_1 == 0x83) && (OutHeader.frame_sync_byte_2 = 0x42)) ? EAVResult::Success : EAVResult::Error;
-    }
-
-    FAVResult ParseColorConfig(FBitstreamReader& Bitstream, Header_t& OutHeader)
-    {
-        if (OutHeader.Profile >= 2)
+			}
+
+			// TODO (william.belcher): More VP9 "P" and "B" frame values. Not needed at the moment
+		}
+
+		return EAVResult::Success;
+	}
+
+	FAVResult ParseFrameSyncCode(FBitstreamReader& Bitstream, Header_t& OutHeader)
+	{
+		Bitstream.Read(OutHeader.frame_sync_byte_0, // u(8)
+			OutHeader.frame_sync_byte_1,			// u(8)
+			OutHeader.frame_sync_byte_2);			// u(8)
+		return ((OutHeader.frame_sync_byte_0 == 0x49) && (OutHeader.frame_sync_byte_1 == 0x83) && (OutHeader.frame_sync_byte_2 = 0x42)) ? EAVResult::Success : EAVResult::Error;
+	}
+
+	FAVResult ParseColorConfig(FBitstreamReader& Bitstream, Header_t& OutHeader)
+	{
+		if (OutHeader.Profile == EVP9Profile::Profile2 || OutHeader.Profile == EVP9Profile::Profile3)
 		{
-            Bitstream.Read(OutHeader.ten_or_twelve_bit);
-            OutHeader.bit_depth = OutHeader.ten_or_twelve_bit ? EBitDepth::k12Bit : EBitDepth::k10Bit;
+			Bitstream.Read(OutHeader.ten_or_twelve_bit);
+			OutHeader.bit_depth = OutHeader.ten_or_twelve_bit ? EBitDepth::k12Bit : EBitDepth::k10Bit;
+		}
+		else
+		{
+			OutHeader.bit_depth = EBitDepth::k8Bit;
 		}
-        else
-        {
-            OutHeader.bit_depth = EBitDepth::k8Bit;
-        }
 
 		Bitstream.Read(OutHeader.color_space); // u(3)
 		if (OutHeader.color_space != EColorSpace::RGB)
 		{
 			Bitstream.Read(OutHeader.color_range); // u(1)
-			if (OutHeader.Profile == 1 || OutHeader.Profile == 3)
+			if (OutHeader.Profile == EVP9Profile::Profile1 || OutHeader.Profile == EVP9Profile::Profile3)
 			{
 				Bitstream.Read(OutHeader.sub_sampling); // u(2)
 				if (Bitstream.ReadBits(1) != 0)
 				{
-                    // Reserved Zero
+					// Reserved Zero
 					return EAVResult::Error;
 				}
 			}
@@ -139,27 +139,27 @@ namespace UE::AVCodecCore::VP9
 		else
 		{
 			OutHeader.color_range = EColorRange::Full;
-			if (OutHeader.Profile == 1 || OutHeader.Profile == 3)
-            {
+			if (OutHeader.Profile == EVP9Profile::Profile1 || OutHeader.Profile == EVP9Profile::Profile3)
+			{
 				OutHeader.sub_sampling = ESubSampling::k444;
 				if (Bitstream.ReadBits(1) != 0)
 				{
-                    // Reserved Zero
+					// Reserved Zero
 					return EAVResult::Error;
 				}
 			}
 		}
 		return EAVResult::Success;
-    }
+	}
 
-    FAVResult ParseFrameSize(FBitstreamReader& Bitstream, Header_t& OutHeader)
-    {
-        Bitstream.Read(OutHeader.frame_width_minus_1,   // u(16)
-                       OutHeader.frame_height_minus_1); // u(16)
+	FAVResult ParseFrameSize(FBitstreamReader& Bitstream, Header_t& OutHeader)
+	{
+		Bitstream.Read(OutHeader.frame_width_minus_1, // u(16)
+			OutHeader.frame_height_minus_1);		  // u(16)
 
-        OutHeader.FrameWidth = OutHeader.frame_width_minus_1 + 1;
-        OutHeader.FrameHeight = OutHeader.frame_height_minus_1 + 1;
+		OutHeader.FrameWidth = OutHeader.frame_width_minus_1 + 1;
+		OutHeader.FrameHeight = OutHeader.frame_height_minus_1 + 1;
 
-        return EAVResult::Success;
-    }
-}
+		return EAVResult::Success;
+	}
+} // namespace UE::AVCodecCore::VP9
diff --git a/Engine/Plugins/Experimental/AVCodecs/AVCodecsCore/Source/AVCodecsCore/Public/Video/CodecUtils/CodecUtilsVP9.h b/Engine/Plugins/Experimental/AVCodecs/AVCodecsCore/Source/AVCodecsCore/Public/Video/CodecUtils/CodecUtilsVP9.h
--- a/Engine/Plugins/Experimental/AVCodecs/AVCodecsCore/Source/AVCodecsCore/Public/Video/CodecUtils/CodecUtilsVP9.h
+++ b/Engine/Plugins/Experimental/AVCodecs/AVCodecsCore/Source/AVCodecsCore/Public/Video/CodecUtils/CodecUtilsVP9.h
@@ -6,107 +6,114 @@
 #include "Utils/BitstreamReader.h"
 #include "AVResult.h"
 
-
 namespace UE::AVCodecCore::VP9
 {
-    enum class EBitDepth : uint8
-    {
-        k8Bit = 8,
-        k10Bit = 10,
-        k12Bit = 12,
-    };
-
-    enum class EColorSpace : uint8
-    {
-        UNKNOWN = 0,    // Unknown (in this case the color space must be signaled
-                     // outside the VP9 bitstream).
-        BT_601 = 1,     // CS_BT_601 Rec. ITU-R BT.601-7
-        BT_709 = 2,     // Rec. ITU-R BT.709-6
-        SMPTE_170 = 3,  // SMPTE-170
-        SMPTE_240 = 4,  // SMPTE-240
-        BT_2020 = 5,    // Rec. ITU-R BT.2020-2
-        RESERVED = 6,   // Reserved
-        RGB = 7,        // sRGB (IEC 61966-2-1)
-    };
-
-    enum class EColorRange : uint8
-    {
-        Studio, // Studio swing:
-                // For BitDepth equals 8:
-                //     Y is between 16 and 235 inclusive.
-                //     U and V are between 16 and 240 inclusive.
-                // For BitDepth equals 10:
-                //     Y is between 64 and 940 inclusive.
-                //     U and V are between 64 and 960 inclusive.
-                // For BitDepth equals 12:
-                //     Y is between 256 and 3760.
-                //     U and V are between 256 and 3840 inclusive.
-        Full    // Full swing; no restriction on Y, U, V values.
-    };
-
-    enum class ESubSampling : uint8
-    {
-        k444,
-        k440,
-        k422,
-        k420,
-    };
-
-    enum class EReferenceFrame : int
-    {
-        None = -1,
-        Intra = 0,
-        Last = 1,
-        Golden = 2,
-        Altref = 3
-    };
-
-    enum class EInterpolationFilter : uint8
-    {
-        EightTap = 0,
-        EightTapSmooth = 1,
-        EightTapSharp = 2,
-        Bilinear = 3,
-        Switchable = 4
-    };
-
-    struct Header_t : public FNalu
-    {
-        U<2> frame_marker = 0;
+	enum class EVP9Profile : uint8
+	{
+		Profile0,
+		Profile1,
+		Profile2,
+		Profile3
+	};
+
+	enum class EBitDepth : uint8
+	{
+		k8Bit = 8,
+		k10Bit = 10,
+		k12Bit = 12,
+	};
+
+	enum class EColorSpace : uint8
+	{
+		UNKNOWN = 0,   // Unknown (in this case the color space must be signaled
+					   // outside the VP9 bitstream).
+		BT_601 = 1,	   // CS_BT_601 Rec. ITU-R BT.601-7
+		BT_709 = 2,	   // Rec. ITU-R BT.709-6
+		SMPTE_170 = 3, // SMPTE-170
+		SMPTE_240 = 4, // SMPTE-240
+		BT_2020 = 5,   // Rec. ITU-R BT.2020-2
+		RESERVED = 6,  // Reserved
+		RGB = 7,	   // sRGB (IEC 61966-2-1)
+	};
+
+	enum class EColorRange : uint8
+	{
+		Studio, // Studio swing:
+				// For BitDepth equals 8:
+				//     Y is between 16 and 235 inclusive.
+				//     U and V are between 16 and 240 inclusive.
+				// For BitDepth equals 10:
+				//     Y is between 64 and 940 inclusive.
+				//     U and V are between 64 and 960 inclusive.
+				// For BitDepth equals 12:
+				//     Y is between 256 and 3760.
+				//     U and V are between 256 and 3840 inclusive.
+		Full	// Full swing; no restriction on Y, U, V values.
+	};
+
+	enum class ESubSampling : uint8
+	{
+		k444,
+		k440,
+		k422,
+		k420,
+	};
+
+	enum class EReferenceFrame : int
+	{
+		None = -1,
+		Intra = 0,
+		Last = 1,
+		Golden = 2,
+		Altref = 3
+	};
+
+	enum class EInterpolationFilter : uint8
+	{
+		EightTap = 0,
+		EightTapSmooth = 1,
+		EightTapSharp = 2,
+		Bilinear = 3,
+		Switchable = 4
+	};
+
+	struct Header_t : public FNalu
+	{
+		U<2> frame_marker = 0;
 		U<1> profile_low_bit = 0;
 		U<1> profile_high_bit = 0;
 		U<1> show_existing_frame = 0;
 		U<3> frame_to_show_map_idx = 0;
 		U<1> frame_type = 0;
-        U<1> show_frame = 0;
+		U<1> show_frame = 0;
 		U<1> error_resilient_mode = 0;
 
-        U<8> frame_sync_byte_0;
-        U<8> frame_sync_byte_1;
-        U<8> frame_sync_byte_2;
+		U<8> frame_sync_byte_0;
+		U<8> frame_sync_byte_1;
+		U<8> frame_sync_byte_2;
 
-        U<1, EBitDepth> bit_depth;
+		U<1, EBitDepth> bit_depth;
 
 		U<1> intra_only = 0;
 		U<2> reset_frame_context = 0;
 
-		U<1> ten_or_twelve_bit = 0;
-		U<3, EColorSpace> color_space;
-		U<1, EColorRange> color_range = EColorRange::Studio;
-        U<2, ESubSampling> sub_sampling;
+		U<1>			   ten_or_twelve_bit = 0;
+		U<3, EColorSpace>  color_space;
+		U<1, EColorRange>  color_range = EColorRange::Studio;
+		U<2, ESubSampling> sub_sampling;
 
-        U<16> frame_width_minus_1;
-        U<16> frame_height_minus_1;
+		U<16> frame_width_minus_1;
+		U<16> frame_height_minus_1;
 
-        // Pre-calculated helpers
-        uint8 Profile = 0;
-        uint16 FrameWidth = 0;
-        uint16 FrameHeight = 0;
-    };
+		// Pre-calculated helpers
+		EVP9Profile Profile = EVP9Profile::Profile0;
+		uint16		FrameWidth = 0;
+		uint16		FrameHeight = 0;
+	};
 
-    FAVResult ParseHeader(FBitstreamReader& Bitstream, Header_t& OutHeader);
+	FAVResult ParseHeader(FBitstreamReader& Bitstream, Header_t& OutHeader);
 
-    FAVResult ParseFrameSyncCode(FBitstreamReader& Bitstream, Header_t& OutHeader);
-    FAVResult ParseColorConfig(FBitstreamReader& Bitstream, Header_t& OutHeader);
-    FAVResult ParseFrameSize(FBitstreamReader& Bitstream, Header_t& OutHeader);
-}
+	FAVResult ParseFrameSyncCode(FBitstreamReader& Bitstream, Header_t& OutHeader);
+	FAVResult ParseColorConfig(FBitstreamReader& Bitstream, Header_t& OutHeader);
+	FAVResult ParseFrameSize(FBitstreamReader& Bitstream, Header_t& OutHeader);
+} // namespace UE::AVCodecCore::VP9
