Add: LibVpx to AVCodecs

diff --git a/Engine/Plugins/Experimental/AVCodecs/AVCodecsCore/Source/AVCodecsCore/Private/Video/CodecUtils/CodecUtilsH264.cpp b/Engine/Plugins/Experimental/AVCodecs/AVCodecsCore/Source/AVCodecsCore/Private/Video/CodecUtils/CodecUtilsH264.cpp
--- a/Engine/Plugins/Experimental/AVCodecs/AVCodecsCore/Source/AVCodecsCore/Private/Video/CodecUtils/CodecUtilsH264.cpp
+++ b/Engine/Plugins/Experimental/AVCodecs/AVCodecsCore/Source/AVCodecsCore/Private/Video/CodecUtils/CodecUtilsH264.cpp
@@ -4,6 +4,7 @@
 
 #include "Containers/Array.h"
 #include "AVResult.h"
+#include "Video/VideoPacket.h"
 
 FH264ProfileDefinition GH264ProfileDefinitions[static_cast<uint8>(EH264Profile::MAX)] = {
 	{ EH264Profile::Auto, UE::AVCodecCore::H264::EH264ProfileIDC::Auto, UE::AVCodecCore::H264::EH264ConstraintFlag::None, TEXT("Auto") },
diff --git a/Engine/Plugins/Experimental/AVCodecs/AVCodecsCore/Source/AVCodecsCore/Private/Video/CodecUtils/CodecUtilsVP9.cpp b/Engine/Plugins/Experimental/AVCodecs/AVCodecsCore/Source/AVCodecsCore/Private/Video/CodecUtils/CodecUtilsVP9.cpp
--- a/Engine/Plugins/Experimental/AVCodecs/AVCodecsCore/Source/AVCodecsCore/Private/Video/CodecUtils/CodecUtilsVP9.cpp
+++ b/Engine/Plugins/Experimental/AVCodecs/AVCodecsCore/Source/AVCodecsCore/Private/Video/CodecUtils/CodecUtilsVP9.cpp
@@ -2,6 +2,8 @@
 
 #include "Video/CodecUtils/CodecUtilsVP9.h"
 
+#include "Video/VideoEncoder.h"
+
 namespace UE::AVCodecCore::VP9
 {
 	FAVResult ParseHeader(FBitstreamReader& Bitstream, Header_t& OutHeader)
@@ -14,9 +16,9 @@ namespace UE::AVCodecCore::VP9
 
 		Bitstream.Read(OutHeader.profile_low_bit, // u(1)
 			OutHeader.profile_high_bit);		  // u(1)
-		OutHeader.Profile = static_cast<EVP9Profile>((OutHeader.profile_high_bit << 1) + OutHeader.profile_low_bit);
+		OutHeader.Profile = static_cast<EProfile>((OutHeader.profile_high_bit << 1) + OutHeader.profile_low_bit);
 
-		if (OutHeader.Profile == EVP9Profile::Profile3)
+		if (OutHeader.Profile == EProfile::Profile3)
 		{
 			if (Bitstream.ReadBits(1) != 0)
 			{
@@ -108,7 +110,7 @@ namespace UE::AVCodecCore::VP9
 
 	FAVResult ParseColorConfig(FBitstreamReader& Bitstream, Header_t& OutHeader)
 	{
-		if (OutHeader.Profile == EVP9Profile::Profile2 || OutHeader.Profile == EVP9Profile::Profile3)
+		if (OutHeader.Profile == EProfile::Profile2 || OutHeader.Profile == EProfile::Profile3)
 		{
 			Bitstream.Read(OutHeader.ten_or_twelve_bit);
 			OutHeader.bit_depth = OutHeader.ten_or_twelve_bit ? EBitDepth::k12Bit : EBitDepth::k10Bit;
@@ -122,7 +124,7 @@ namespace UE::AVCodecCore::VP9
 		if (OutHeader.color_space != EColorSpace::RGB)
 		{
 			Bitstream.Read(OutHeader.color_range); // u(1)
-			if (OutHeader.Profile == EVP9Profile::Profile1 || OutHeader.Profile == EVP9Profile::Profile3)
+			if (OutHeader.Profile == EProfile::Profile1 || OutHeader.Profile == EProfile::Profile3)
 			{
 				Bitstream.Read(OutHeader.sub_sampling); // u(2)
 				if (Bitstream.ReadBits(1) != 0)
@@ -139,7 +141,7 @@ namespace UE::AVCodecCore::VP9
 		else
 		{
 			OutHeader.color_range = EColorRange::Full;
-			if (OutHeader.Profile == EVP9Profile::Profile1 || OutHeader.Profile == EVP9Profile::Profile3)
+			if (OutHeader.Profile == EProfile::Profile1 || OutHeader.Profile == EProfile::Profile3)
 			{
 				OutHeader.sub_sampling = ESubSampling::k444;
 				if (Bitstream.ReadBits(1) != 0)
@@ -162,4 +164,64 @@ namespace UE::AVCodecCore::VP9
 
 		return EAVResult::Success;
 	}
+
+	EInterLayerPrediction ScalabilityModeToInterLayerPredMode(EScalabilityMode ScalabilityMode)
+	{
+		switch (ScalabilityMode)
+		{
+			case EScalabilityMode::L1T1:
+			case EScalabilityMode::L1T2:
+			case EScalabilityMode::L1T3:
+			case EScalabilityMode::L2T1:
+			case EScalabilityMode::L2T1h:
+				return EInterLayerPrediction::On;
+			case EScalabilityMode::L2T1_KEY:
+				return EInterLayerPrediction::OnKeyPicture;
+			case EScalabilityMode::L2T2:
+			case EScalabilityMode::L2T2h:
+				return EInterLayerPrediction::On;
+			case EScalabilityMode::L2T2_KEY:
+			case EScalabilityMode::L2T2_KEY_SHIFT:
+				return EInterLayerPrediction::OnKeyPicture;
+			case EScalabilityMode::L2T3:
+			case EScalabilityMode::L2T3h:
+				return EInterLayerPrediction::On;
+			case EScalabilityMode::L2T3_KEY:
+				return EInterLayerPrediction::OnKeyPicture;
+			case EScalabilityMode::L3T1:
+			case EScalabilityMode::L3T1h:
+				return EInterLayerPrediction::On;
+			case EScalabilityMode::L3T1_KEY:
+				return EInterLayerPrediction::OnKeyPicture;
+			case EScalabilityMode::L3T2:
+			case EScalabilityMode::L3T2h:
+				return EInterLayerPrediction::On;
+			case EScalabilityMode::L3T2_KEY:
+				return EInterLayerPrediction::OnKeyPicture;
+			case EScalabilityMode::L3T3:
+			case EScalabilityMode::L3T3h:
+				return EInterLayerPrediction::On;
+			case EScalabilityMode::L3T3_KEY:
+				return EInterLayerPrediction::OnKeyPicture;
+			case EScalabilityMode::S2T1:
+			case EScalabilityMode::S2T1h:
+			case EScalabilityMode::S2T2:
+			case EScalabilityMode::S2T2h:
+			case EScalabilityMode::S2T3:
+			case EScalabilityMode::S2T3h:
+			case EScalabilityMode::S3T1:
+			case EScalabilityMode::S3T1h:
+			case EScalabilityMode::S3T2:
+			case EScalabilityMode::S3T2h:
+			case EScalabilityMode::S3T3:
+			case EScalabilityMode::S3T3h:
+			case EScalabilityMode::None:
+				return EInterLayerPrediction::Off;
+		}
+
+		// All scalability modes should be accounted for in the above switch
+		unimplemented();
+		return EInterLayerPrediction::Off;
+	}
+
 } // namespace UE::AVCodecCore::VP9
diff --git a/Engine/Plugins/Experimental/AVCodecs/AVCodecsCore/Source/AVCodecsCore/Private/Video/Resources/VideoResourceCPU.cpp b/Engine/Plugins/Experimental/AVCodecs/AVCodecsCore/Source/AVCodecsCore/Private/Video/Resources/VideoResourceCPU.cpp
--- a/Engine/Plugins/Experimental/AVCodecs/AVCodecsCore/Source/AVCodecsCore/Private/Video/Resources/VideoResourceCPU.cpp
+++ b/Engine/Plugins/Experimental/AVCodecs/AVCodecsCore/Source/AVCodecsCore/Private/Video/Resources/VideoResourceCPU.cpp
@@ -4,6 +4,7 @@
 
 REGISTER_TYPEID(FVideoContextCPU);
 REGISTER_TYPEID(FVideoResourceCPU);
+REGISTER_TYPEID(FResolvableVideoResourceCPU);
 
 FVideoContextCPU::FVideoContextCPU()
 {
@@ -24,3 +25,60 @@ FAVResult FVideoResourceCPU::Validate() const
 
 	return EAVResult::Success;
 }
+
+TSharedPtr<FVideoResourceCPU> FResolvableVideoResourceCPU::TryResolve(TSharedPtr<FAVDevice> const& Device, FVideoDescriptor const& Descriptor)
+{
+	return MakeShared<FVideoResourceCPU>(Device.ToSharedRef(), nullptr, FAVLayout(GetStrideInterleavedOrLuma(Descriptor), 0, GetSize(Descriptor)), Descriptor);
+}
+
+uint32 FResolvableVideoResourceCPU::GetStrideInterleavedOrLuma(FVideoDescriptor const& Descriptor)
+{
+	switch (Descriptor.Format)
+	{
+		case EVideoFormat::BGRA:
+		case EVideoFormat::ABGR10:
+		case EVideoFormat::YUV420:
+		case EVideoFormat::YUV444:
+		case EVideoFormat::NV12:
+			return Descriptor.Width;
+		case EVideoFormat::YUV444_16:
+		case EVideoFormat::P010:
+			return Descriptor.Width * 2;
+		case EVideoFormat::R8:
+			// TODO (william.belcher)
+			return 0;
+		case EVideoFormat::G16:
+			// TODO (william.belcher)
+			return 0;
+		default:
+			return 0;
+	}
+}
+
+uint32 FResolvableVideoResourceCPU::GetSize(FVideoDescriptor const& Descriptor)
+{
+	switch (Descriptor.Format)
+	{
+		case EVideoFormat::BGRA:
+		case EVideoFormat::ABGR10:
+			return Descriptor.Width * Descriptor.Height * 4;
+		case EVideoFormat::YUV420:
+			return Descriptor.Width * Descriptor.Height + (2 * ((Descriptor.Width + 1) / 2) * ((Descriptor.Height + 1) / 2));
+		case EVideoFormat::YUV444:
+			return Descriptor.Width * Descriptor.Height * 3;
+		case EVideoFormat::YUV444_16:
+			return Descriptor.Width * Descriptor.Height * 6;
+		case EVideoFormat::NV12:
+			return Descriptor.Width * Descriptor.Height + (Descriptor.Width * ((Descriptor.Height + 1) / 2));
+		case EVideoFormat::P010:
+			return Descriptor.Width * Descriptor.Height + (2 * ((Descriptor.Width + 1) / 2) * ((Descriptor.Height + 1) / 2)) * 2;
+		case EVideoFormat::R8:
+			// TODO (william.belcher)
+			return 0;
+		case EVideoFormat::G16:
+			// TODO (william.belcher)
+			return 0;
+		default:
+			return 0;
+	}
+}
diff --git a/Engine/Plugins/Experimental/AVCodecs/AVCodecsCore/Source/AVCodecsCore/Public/Video/CodecUtils/CodecUtilsH264.h b/Engine/Plugins/Experimental/AVCodecs/AVCodecsCore/Source/AVCodecsCore/Public/Video/CodecUtils/CodecUtilsH264.h
--- a/Engine/Plugins/Experimental/AVCodecs/AVCodecsCore/Source/AVCodecsCore/Public/Video/CodecUtils/CodecUtilsH264.h
+++ b/Engine/Plugins/Experimental/AVCodecs/AVCodecsCore/Source/AVCodecsCore/Public/Video/CodecUtils/CodecUtilsH264.h
@@ -6,9 +6,9 @@
 #include "Utils/BitstreamReader.h"
 
 #include "AVResult.h"
-#include "Video/VideoPacket.h"
 
 struct FH264ProfileDefinition;
+struct FVideoPacket;
 
 // Easy to understand user facing profiles from the T-Rec-H.264-201304 spec
 // there is another enum in CodecUtilsH264 that aligns the profiles with PIC and constraints
@@ -55,6 +55,16 @@ enum class EH264EntropyCodingMode : uint8
 	CAVLC
 };
 
+// Packetization modes are defined in RFC 6184 section 6
+// Due to the structure containing this being initialized with zeroes
+// in some places, and mode 1 being default, mode 1 needs to have the value
+// zero.
+enum class EH264PacketizationMode : uint8
+{
+	NonInterleaved = 0, // Mode 1 - STAP-A, FU-A is allowed
+	SingleNalUnit		// Mode 0 - only single NALU allowed
+};
+
 extern AVCODECSCORE_API FH264ProfileDefinition GH264ProfileDefinitions[static_cast<uint8>(EH264Profile::MAX)];
 
 namespace UE::AVCodecCore::H264
@@ -117,9 +127,9 @@ namespace UE::AVCodecCore::H264
 		High422 = 122,
 		// High422Intra = 122,					with constraint flag 3 set
 		StereoHigh = 128,
-        MultiresolutionFrameCompatibleHigh = 134,
+		MultiresolutionFrameCompatibleHigh = 134,
 		MultiviewDepthHigh = 138,
-        EnhancedMultiviewDepthHigh = 139,
+		EnhancedMultiviewDepthHigh = 139,
 		High444 = 244,
 		// High444Intra = 244,					with constraint flag 3 set
 	};
@@ -222,10 +232,10 @@ namespace UE::AVCodecCore::H264
 
 struct FH264ProfileDefinition
 {
-	EH264Profile Profile;
-	UE::AVCodecCore::H264::EH264ProfileIDC PIDC;
+	EH264Profile							   Profile;
+	UE::AVCodecCore::H264::EH264ProfileIDC	   PIDC;
 	UE::AVCodecCore::H264::EH264ConstraintFlag ConstraintFlags;
-	const TCHAR* Name;
+	const TCHAR*							   Name;
 };
 
 namespace UE::AVCodecCore::H264
@@ -255,10 +265,10 @@ namespace UE::AVCodecCore::H264
 
 	struct FNaluH264
 	{
-		uint64 Start, Size;
-		uint8 StartCodeSize;
-		uint8 RefIdc;
-		ENaluType Type;
+		uint64		 Start, Size;
+		uint8		 StartCodeSize;
+		uint8		 RefIdc;
+		ENaluType	 Type;
 		const uint8* Data;
 	};
 
@@ -355,79 +365,79 @@ namespace UE::AVCodecCore::H264
 	struct SPS_t : public FNalu
 	{
 	public:
-		U<8, EH264ProfileIDC> profile_idc = EH264ProfileIDC::Auto;
+		U<8, EH264ProfileIDC>	  profile_idc = EH264ProfileIDC::Auto;
 		U<8, EH264ConstraintFlag> constraint_flags = EH264ConstraintFlag::None;
-		U<8> level_idc;							   // + profile_idc + constraint_flags = EH264Level
-		UE seq_parameter_set_id;				   // 0...31 inclusive
-		UE chroma_format_idc = 1;				   // 0...3 inclusive
-		U<1> separate_colour_plane_flag;		   // bool
-		UE bit_depth_luma_minus8;				   // 0...6 inclusive
-		UE bit_depth_chroma_minus8;				   // 0...6 inclusive
-		U<1> qpprime_y_zero_transform_bypass_flag; // 0...1 inclusive
-		U<1> seq_scaling_matrix_present_flag;	   // bool
-		U<1> seq_scaling_list_present_flag[12];	   // bool array
-		uint8 ScalingList4x4[6][16];			   // -2^(1+BitDepth)...2^(1+BitDepth)-1 inclusive
-		uint8 ScalingList8x8[6][64];			   // -2^(1+BitDepth)...2^(1+BitDepth)-1 inclusive
-		UE log2_max_frame_num_minus4;			   // 0...12 inclusive
-		UE pic_order_cnt_type;					   // 0...3 inclusive
-		UE log2_max_pic_order_cnt_lsb_minus4;	   // 0...12 inclusive
-		U<1> delta_pic_order_always_zero_flag;	   // bool
-		SE offset_for_non_ref_pic;				   // −2^(31)+1...2^(31)−1 inclusive
-		SE offset_for_top_to_bottom_field;		   // −2^(31)+1...2^(31)−1 inclusive
-		UE num_ref_frames_in_pic_order_cnt_cycle;  // 0...255 inclusive
-		SE offset_for_ref_frame[255];			   // −2^(31)+1...2^(31)−1 inclusive
-		UE max_num_ref_frames;					   // 0...MaxDpbFrames inclusive MaxDpbFrames = Min( MaxDecPicBufSizeMacro / ( PicWidthMacro * FrameHeightMacro ), 16) (A.3.1 or A.3.2)
-		U<1> gaps_in_frame_num_value_allowed_flag; // bool
-		UE pic_width_in_mbs_minus1;				   // + 1 = Width in macroblocks
-		UE pic_height_in_map_units_minus1;		   // + 1 = Height
-		U<1> frame_mbs_only_flag;				   // bool
-		U<1> mb_adaptive_frame_field_flag;		   // bool
-		U<1> direct_8x8_inference_flag;			   // bool
-		U<1> frame_cropping_flag;				   // bool
-		UE frame_crop_left_offset;
-		UE frame_crop_right_offset;
-		UE frame_crop_top_offset;
-		UE frame_crop_bottom_offset;
-		U<1> vui_parameters_present_flag; // bool
+		U<8>					  level_idc;							 // + profile_idc + constraint_flags = EH264Level
+		UE						  seq_parameter_set_id;					 // 0...31 inclusive
+		UE						  chroma_format_idc = 1;				 // 0...3 inclusive
+		U<1>					  separate_colour_plane_flag;			 // bool
+		UE						  bit_depth_luma_minus8;				 // 0...6 inclusive
+		UE						  bit_depth_chroma_minus8;				 // 0...6 inclusive
+		U<1>					  qpprime_y_zero_transform_bypass_flag;	 // 0...1 inclusive
+		U<1>					  seq_scaling_matrix_present_flag;		 // bool
+		U<1>					  seq_scaling_list_present_flag[12];	 // bool array
+		uint8					  ScalingList4x4[6][16];				 // -2^(1+BitDepth)...2^(1+BitDepth)-1 inclusive
+		uint8					  ScalingList8x8[6][64];				 // -2^(1+BitDepth)...2^(1+BitDepth)-1 inclusive
+		UE						  log2_max_frame_num_minus4;			 // 0...12 inclusive
+		UE						  pic_order_cnt_type;					 // 0...3 inclusive
+		UE						  log2_max_pic_order_cnt_lsb_minus4;	 // 0...12 inclusive
+		U<1>					  delta_pic_order_always_zero_flag;		 // bool
+		SE						  offset_for_non_ref_pic;				 // −2^(31)+1...2^(31)−1 inclusive
+		SE						  offset_for_top_to_bottom_field;		 // −2^(31)+1...2^(31)−1 inclusive
+		UE						  num_ref_frames_in_pic_order_cnt_cycle; // 0...255 inclusive
+		SE						  offset_for_ref_frame[255];			 // −2^(31)+1...2^(31)−1 inclusive
+		UE						  max_num_ref_frames;					 // 0...MaxDpbFrames inclusive MaxDpbFrames = Min( MaxDecPicBufSizeMacro / ( PicWidthMacro * FrameHeightMacro ), 16) (A.3.1 or A.3.2)
+		U<1>					  gaps_in_frame_num_value_allowed_flag;	 // bool
+		UE						  pic_width_in_mbs_minus1;				 // + 1 = Width in macroblocks
+		UE						  pic_height_in_map_units_minus1;		 // + 1 = Height
+		U<1>					  frame_mbs_only_flag;					 // bool
+		U<1>					  mb_adaptive_frame_field_flag;			 // bool
+		U<1>					  direct_8x8_inference_flag;			 // bool
+		U<1>					  frame_cropping_flag;					 // bool
+		UE						  frame_crop_left_offset;
+		UE						  frame_crop_right_offset;
+		UE						  frame_crop_top_offset;
+		UE						  frame_crop_bottom_offset;
+		U<1>					  vui_parameters_present_flag; // bool
 		// VUI Parameters
-		U<1> aspect_ratio_info_present_flag; // bool
+		U<1>					  aspect_ratio_info_present_flag; // bool
 		U<8, EH264AspectRatioIDC> aspect_ratio_idc = EH264AspectRatioIDC::Unspecified;
-		U<16> sar_width;					 // Width in arbitrary unit
-		U<16> sar_height;					 // Height in same arbitrary unit as Width
-		U<1> overscan_info_present_flag;	 // bool
-		U<1> overscan_appropriate_flag;		 // bool
-		U<1> video_signal_type_present_flag; // bool
-		U<3, EH264VideoFormat> video_format = EH264VideoFormat::Unspecified;
-		U<1> video_full_range_flag;					  // bool
-		U<1> colour_description_present_flag;		  // bool
-		U<8> colour_primaries = 2;					  // 0...255 inclusive but see Table E-3 as many values are Reserved
-		U<8> transfer_characteristics = 2;			  // 0...255 inclusive but see Table E-4 as many values are Reserved
-		U<8> matrix_coefficients = 2;				  // 0...255 inclusive but see Table E-5 as many values are Reserved
-		U<1> chroma_loc_info_present_flag;			  // bool
-		UE chroma_sample_loc_type_top_field;		  // 0...5 inclusive
-		UE chroma_sample_loc_type_bottom_field;		  // 0...5 inclusive
-		U<1> timing_info_present_flag;				  // bool
-		U<32> num_units_in_tick = 1;				  // > 0
-		U<32> time_scale = 1000;					  // defaulted to ms
-		U<1> fixed_frame_rate_flag;					  // bool
-		U<1> nal_hrd_parameters_present_flag;		  // bool
-		U<1> vcl_hrd_parameters_present_flag;		  // bool
-		U<1> low_delay_hrd_flag;					  // bool
-		U<1> pic_struct_present_flag;				  // bool
-		U<1> bitstream_restriction_flag;			  // bool
-		U<1> motion_vectors_over_pic_boundaries_flag; // bool
-		UE max_bytes_per_pic_denom;
-		UE max_bits_per_mb_denom;
-		UE log2_max_mv_length_horizontal;
-		UE log2_max_mv_length_vertical;
-		UE max_num_reorder_frames;
-		UE max_dec_frame_buffering;
+		U<16>					  sar_width;					  // Width in arbitrary unit
+		U<16>					  sar_height;					  // Height in same arbitrary unit as Width
+		U<1>					  overscan_info_present_flag;	  // bool
+		U<1>					  overscan_appropriate_flag;	  // bool
+		U<1>					  video_signal_type_present_flag; // bool
+		U<3, EH264VideoFormat>	  video_format = EH264VideoFormat::Unspecified;
+		U<1>					  video_full_range_flag;				   // bool
+		U<1>					  colour_description_present_flag;		   // bool
+		U<8>					  colour_primaries = 2;					   // 0...255 inclusive but see Table E-3 as many values are Reserved
+		U<8>					  transfer_characteristics = 2;			   // 0...255 inclusive but see Table E-4 as many values are Reserved
+		U<8>					  matrix_coefficients = 2;				   // 0...255 inclusive but see Table E-5 as many values are Reserved
+		U<1>					  chroma_loc_info_present_flag;			   // bool
+		UE						  chroma_sample_loc_type_top_field;		   // 0...5 inclusive
+		UE						  chroma_sample_loc_type_bottom_field;	   // 0...5 inclusive
+		U<1>					  timing_info_present_flag;				   // bool
+		U<32>					  num_units_in_tick = 1;				   // > 0
+		U<32>					  time_scale = 1000;					   // defaulted to ms
+		U<1>					  fixed_frame_rate_flag;				   // bool
+		U<1>					  nal_hrd_parameters_present_flag;		   // bool
+		U<1>					  vcl_hrd_parameters_present_flag;		   // bool
+		U<1>					  low_delay_hrd_flag;					   // bool
+		U<1>					  pic_struct_present_flag;				   // bool
+		U<1>					  bitstream_restriction_flag;			   // bool
+		U<1>					  motion_vectors_over_pic_boundaries_flag; // bool
+		UE						  max_bytes_per_pic_denom;
+		UE						  max_bits_per_mb_denom;
+		UE						  log2_max_mv_length_horizontal;
+		UE						  log2_max_mv_length_vertical;
+		UE						  max_num_reorder_frames;
+		UE						  max_dec_frame_buffering;
 		// HRD Parameters
-		UE cpb_cnt_minus1; // 0...31 inclusive
+		UE	 cpb_cnt_minus1; // 0...31 inclusive
 		U<4> bit_rate_scale;
 		U<4> cpb_size_scale;
-		UE bit_rate_value_minus1[31];					   // 0...2^(32)−2 inclusive TODO (aidan) might need an infer stage to fill out based on profile
-		UE cpb_size_value_minus1[31];					   // 0...2^(32)−2 inclusive TODO (aidan) might need an infer stage to fill out based on profile
+		UE	 bit_rate_value_minus1[31];					   // 0...2^(32)−2 inclusive TODO (aidan) might need an infer stage to fill out based on profile
+		UE	 cpb_size_value_minus1[31];					   // 0...2^(32)−2 inclusive TODO (aidan) might need an infer stage to fill out based on profile
 		U<1> cbr_flag[31];								   // bool
 		U<5> initial_cpb_removal_delay_length_minus1 = 23; //
 		U<5> cpb_removal_delay_length_minus1 = 23;
@@ -452,71 +462,71 @@ namespace UE::AVCodecCore::H264
 
 	struct PPS_t : public FNalu
 	{
-		UE pic_parameter_set_id;
-		UE seq_parameter_set_id;
-		U<1> entropy_coding_mode_flag;
-		U<1> bottom_field_pic_order_in_frame_present_flag;
-		UE num_slice_groups_minus1;
-		UE slice_group_map_type;
-		TArray<UE> run_length_minus1; // hard to justify preallocating this as it could be PicWidthInMbs * PicHeightInMapUnits - 1
-		TArray<UE> top_left;		  // hard to justify preallocating this as it could be PicWidthInMbs * PicHeightInMapUnits - 1
-		TArray<UE> bottom_right;	  // hard to justify preallocating this as it could be PicWidthInMbs * PicHeightInMapUnits - 1
-		U<1> slice_group_change_direction_flag;
-		UE slice_group_change_rate_minus1;
-		UE pic_size_in_map_units_minus1;
+		UE			pic_parameter_set_id;
+		UE			seq_parameter_set_id;
+		U<1>		entropy_coding_mode_flag;
+		U<1>		bottom_field_pic_order_in_frame_present_flag;
+		UE			num_slice_groups_minus1;
+		UE			slice_group_map_type;
+		TArray<UE>	run_length_minus1; // hard to justify preallocating this as it could be PicWidthInMbs * PicHeightInMapUnits - 1
+		TArray<UE>	top_left;		   // hard to justify preallocating this as it could be PicWidthInMbs * PicHeightInMapUnits - 1
+		TArray<UE>	bottom_right;	   // hard to justify preallocating this as it could be PicWidthInMbs * PicHeightInMapUnits - 1
+		U<1>		slice_group_change_direction_flag;
+		UE			slice_group_change_rate_minus1;
+		UE			pic_size_in_map_units_minus1;
 		TArray<U<>> slice_group_id; // hard to justify preallocating this as it could be pic_size_in_map_units_minus1 + 1
-		UE num_ref_idx_l0_default_active_minus1;
-		UE num_ref_idx_l1_default_active_minus1;
-		U<1> weighted_pred_flag;
-		U<2> weighted_bipred_idc;
-		SE pic_init_qp_minus26;
-		SE pic_init_qs_minus26;
-		SE chroma_qp_index_offset;
-		U<1> deblocking_filter_control_present_flag;
-		U<1> constrained_intra_pred_flag;
-		U<1> redundant_pic_cnt_present_flag;
-		U<1> transform_8x8_mode_flag;
-		U<1> pic_scaling_matrix_present_flag;
-		U<1> pic_scaling_list_present_flag[12];
-		uint8 ScalingList4x4[6][16]; // -2^(1+BitDepth)...2^(1+BitDepth)-1 inclusive
-		uint8 ScalingList8x8[6][64]; // -2^(1+BitDepth)...2^(1+BitDepth)-1 inclusive
-		SE second_chroma_qp_index_offset;
+		UE			num_ref_idx_l0_default_active_minus1;
+		UE			num_ref_idx_l1_default_active_minus1;
+		U<1>		weighted_pred_flag;
+		U<2>		weighted_bipred_idc;
+		SE			pic_init_qp_minus26;
+		SE			pic_init_qs_minus26;
+		SE			chroma_qp_index_offset;
+		U<1>		deblocking_filter_control_present_flag;
+		U<1>		constrained_intra_pred_flag;
+		U<1>		redundant_pic_cnt_present_flag;
+		U<1>		transform_8x8_mode_flag;
+		U<1>		pic_scaling_matrix_present_flag;
+		U<1>		pic_scaling_list_present_flag[12];
+		uint8		ScalingList4x4[6][16]; // -2^(1+BitDepth)...2^(1+BitDepth)-1 inclusive
+		uint8		ScalingList8x8[6][64]; // -2^(1+BitDepth)...2^(1+BitDepth)-1 inclusive
+		SE			second_chroma_qp_index_offset;
 	};
 
 	FAVResult ParsePPS(FBitstreamReader& Bitstream, FNaluH264 const& InNaluInfo, TMap<uint32, SPS_t> const& InMapSPS, TMap<uint32, PPS_t>& OutMapPPS);
 
 	struct Slice_t : public FNalu
 	{
-		UE first_mb_in_slice;
-		UE slice_type;
-		UE pic_parameter_set_id;
+		UE	 first_mb_in_slice;
+		UE	 slice_type;
+		UE	 pic_parameter_set_id;
 		U<2> colour_plane_id;
-		U<> frame_num;
+		U<>	 frame_num;
 		U<1> field_pic_flag;
 		U<1> bottom_field_flag;
-		UE idr_pic_id;
-		U<> pic_order_cnt_lsb;
-		SE delta_pic_order_cnt_bottom;
-		SE delta_pic_order_cnt[2];
-		UE redundant_pic_cnt;
+		UE	 idr_pic_id;
+		U<>	 pic_order_cnt_lsb;
+		SE	 delta_pic_order_cnt_bottom;
+		SE	 delta_pic_order_cnt[2];
+		UE	 redundant_pic_cnt;
 		U<1> direct_spatial_mv_pred_flag;
 		U<1> num_ref_idx_active_override_flag;
-		UE num_ref_idx_l0_active_minus1;
-		UE num_ref_idx_l1_active_minus1;
-		UE cabac_init_idc;
-		SE slice_qp_delta;
+		UE	 num_ref_idx_l0_active_minus1;
+		UE	 num_ref_idx_l1_active_minus1;
+		UE	 cabac_init_idc;
+		SE	 slice_qp_delta;
 		U<1> sp_for_switch_flag;
-		SE slice_qs_delta;
-		UE disable_deblocking_filter_idc;
-		SE slice_alpha_c0_offset_div2;
-		SE slice_beta_offset_div2;
-		U<> slice_group_change_cycle;
+		SE	 slice_qs_delta;
+		UE	 disable_deblocking_filter_idc;
+		SE	 slice_alpha_c0_offset_div2;
+		SE	 slice_beta_offset_div2;
+		U<>	 slice_group_change_cycle;
 
 		// ref_pic_list_modification
 		struct RefPic_t
 		{
 			U<1> bIsLongTerm;
-			UE pic_num;
+			UE	 pic_num;
 		};
 
 		U<1> ref_pic_list_modification_flag_l0;
@@ -532,10 +542,10 @@ namespace UE::AVCodecCore::H264
 
 		U<1> adaptive_ref_pic_marking_mode_flag;
 
-        UE difference_of_pic_nums_minus1;
-        UE long_term_pic_num;
-        UE long_term_frame_idx;
-        UE max_long_term_frame_idx_plus1;
+		UE difference_of_pic_nums_minus1;
+		UE long_term_pic_num;
+		UE long_term_frame_idx;
+		UE max_long_term_frame_idx_plus1;
 	};
 
 	FAVResult ParseSliceHeader(FBitstreamReader& Bitstream, FNaluH264 const& InNaluInfo, TMap<uint32, SPS_t> const& InMapSPS, TMap<uint32, PPS_t> const& InMapPPS, Slice_t& OutSlice);
diff --git a/Engine/Plugins/Experimental/AVCodecs/AVCodecsCore/Source/AVCodecsCore/Public/Video/CodecUtils/CodecUtilsVP9.h b/Engine/Plugins/Experimental/AVCodecs/AVCodecsCore/Source/AVCodecsCore/Public/Video/CodecUtils/CodecUtilsVP9.h
--- a/Engine/Plugins/Experimental/AVCodecs/AVCodecsCore/Source/AVCodecsCore/Public/Video/CodecUtils/CodecUtilsVP9.h
+++ b/Engine/Plugins/Experimental/AVCodecs/AVCodecsCore/Source/AVCodecsCore/Public/Video/CodecUtils/CodecUtilsVP9.h
@@ -6,9 +6,23 @@
 #include "Utils/BitstreamReader.h"
 #include "AVResult.h"
 
+enum class EScalabilityMode : uint8;
+
 namespace UE::AVCodecCore::VP9
 {
-	enum class EVP9Profile : uint8
+	constexpr int16	 MaxOneBytePictureId = 0x7F;   // 7 bits
+	constexpr int16	 MaxTwoBytePictureId = 0x7FFF; // 15 bits
+	constexpr uint8	 NoSpatialIdx = 0xFF;
+	constexpr uint8	 NoGofIdx = 0xFF;
+	constexpr uint8	 NumBuffers = 8;
+	constexpr size_t MaxRefPics = 3;
+	constexpr size_t MaxFramesInGof = 0xFF; // 8 bits
+	constexpr size_t MaxNumberOfSpatialLayers = 8;
+
+	constexpr size_t MinSpatialLayerLongSideLength = 240;
+	constexpr size_t MinSpatialLayerShortSideLength = 135;
+
+	enum class EProfile : uint8
 	{
 		Profile0,
 		Profile1,
@@ -16,6 +30,80 @@ namespace UE::AVCodecCore::VP9
 		Profile3
 	};
 
+	enum class EInterLayerPrediction : uint8
+	{
+		Off,
+		On,
+		OnKeyPicture
+	};
+
+	enum class ETemporalStructureMode : uint8
+	{
+		TemporalStructureMode1, // 1 temporal layer structure - i.e., IPPP...
+		TemporalStructureMode2, // 2 temporal layers 01...
+		TemporalStructureMode3, // 3 temporal layers 0212...
+	};
+
+	struct FGroupOfFramesInfo
+	{
+		void SetGofInfo(ETemporalStructureMode tm)
+		{
+			switch (tm)
+			{
+				case ETemporalStructureMode::TemporalStructureMode1:
+					NumFramesInGof = 1;
+					TemporalIdx[0] = 0;
+					TemporalUpSwitch[0] = true;
+					NumRefPics[0] = 1;
+					PidDiff[0][0] = 1;
+					break;
+				case ETemporalStructureMode::TemporalStructureMode2:
+					NumFramesInGof = 2;
+					TemporalIdx[0] = 0;
+					TemporalUpSwitch[0] = true;
+					NumRefPics[0] = 1;
+					PidDiff[0][0] = 2;
+
+					TemporalIdx[1] = 1;
+					TemporalUpSwitch[1] = true;
+					NumRefPics[1] = 1;
+					PidDiff[1][0] = 1;
+					break;
+				case ETemporalStructureMode::TemporalStructureMode3:
+					NumFramesInGof = 4;
+					TemporalIdx[0] = 0;
+					TemporalUpSwitch[0] = true;
+					NumRefPics[0] = 1;
+					PidDiff[0][0] = 4;
+
+					TemporalIdx[1] = 2;
+					TemporalUpSwitch[1] = true;
+					NumRefPics[1] = 1;
+					PidDiff[1][0] = 1;
+
+					TemporalIdx[2] = 1;
+					TemporalUpSwitch[2] = true;
+					NumRefPics[2] = 1;
+					PidDiff[2][0] = 2;
+
+					TemporalIdx[3] = 2;
+					TemporalUpSwitch[3] = true;
+					NumRefPics[3] = 1;
+					PidDiff[3][0] = 1;
+					break;
+				default:
+					checkNoEntry();
+			}
+		}
+
+		size_t NumFramesInGof;
+		uint8  TemporalIdx[MaxFramesInGof];
+		bool   TemporalUpSwitch[MaxFramesInGof];
+		uint8  NumRefPics[MaxFramesInGof];
+		uint8  PidDiff[MaxFramesInGof][MaxRefPics];
+		uint16 PidStart;
+	};
+
 	enum class EBitDepth : uint8
 	{
 		k8Bit = 8,
@@ -106,7 +194,7 @@ namespace UE::AVCodecCore::VP9
 		U<16> frame_height_minus_1;
 
 		// Pre-calculated helpers
-		EVP9Profile Profile = EVP9Profile::Profile0;
+		EProfile Profile = EProfile::Profile0;
 		uint16		FrameWidth = 0;
 		uint16		FrameHeight = 0;
 	};
@@ -116,4 +204,7 @@ namespace UE::AVCodecCore::VP9
 	FAVResult ParseFrameSyncCode(FBitstreamReader& Bitstream, Header_t& OutHeader);
 	FAVResult ParseColorConfig(FBitstreamReader& Bitstream, Header_t& OutHeader);
 	FAVResult ParseFrameSize(FBitstreamReader& Bitstream, Header_t& OutHeader);
+
+	AVCODECSCORE_API EInterLayerPrediction ScalabilityModeToInterLayerPredMode(EScalabilityMode ScalabilityMode);
+
 } // namespace UE::AVCodecCore::VP9
diff --git a/Engine/Plugins/Experimental/AVCodecs/AVCodecsCore/Source/AVCodecsCore/Public/Video/Decoders/Configs/VideoDecoderConfigVP9.h b/Engine/Plugins/Experimental/AVCodecs/AVCodecsCore/Source/AVCodecsCore/Public/Video/Decoders/Configs/VideoDecoderConfigVP9.h
--- a/Engine/Plugins/Experimental/AVCodecs/AVCodecsCore/Source/AVCodecsCore/Public/Video/Decoders/Configs/VideoDecoderConfigVP9.h
+++ b/Engine/Plugins/Experimental/AVCodecs/AVCodecsCore/Source/AVCodecsCore/Public/Video/Decoders/Configs/VideoDecoderConfigVP9.h
@@ -10,7 +10,13 @@
  */
 struct AVCODECSCORE_API FVideoDecoderConfigVP9 : public FVideoDecoderConfig
 {
-    FVideoDecoderConfigVP9(EAVPreset Preset = EAVPreset::Default)
+	// Values are parsed from the bitstream during decoding
+	uint32 MaxOutputWidth = 0;
+	uint32 MaxOutputHeight = 0;
+	
+	int32  NumberOfCores = 0;
+
+	FVideoDecoderConfigVP9(EAVPreset Preset = EAVPreset::Default)
 		: FVideoDecoderConfig(Preset)
 	{
 	}
diff --git a/Engine/Plugins/Experimental/AVCodecs/AVCodecsCore/Source/AVCodecsCore/Public/Video/Encoders/Configs/VideoEncoderConfigAV1.h b/Engine/Plugins/Experimental/AVCodecs/AVCodecsCore/Source/AVCodecsCore/Public/Video/Encoders/Configs/VideoEncoderConfigAV1.h
--- a/Engine/Plugins/Experimental/AVCodecs/AVCodecsCore/Source/AVCodecsCore/Public/Video/Encoders/Configs/VideoEncoderConfigAV1.h
+++ b/Engine/Plugins/Experimental/AVCodecs/AVCodecsCore/Source/AVCodecsCore/Public/Video/Encoders/Configs/VideoEncoderConfigAV1.h
@@ -17,6 +17,7 @@ public:
 
 	uint32 IntraRefreshPeriodFrames = 0;
 	uint32 IntraRefreshCountFrames = 0;
+	uint32 BitDepthMinus8 = 0;
 
 	FVideoEncoderConfigAV1(EAVPreset Preset = EAVPreset::Default)
 		: FVideoEncoderConfig(Preset)
diff --git a/Engine/Plugins/Experimental/AVCodecs/AVCodecsCore/Source/AVCodecsCore/Public/Video/Encoders/SVC/ScalabilityStructureFull.h b/Engine/Plugins/Experimental/AVCodecs/AVCodecsCore/Source/AVCodecsCore/Public/Video/Encoders/SVC/ScalabilityStructureFull.h
new file mode 100644
--- /dev/null
+++ b/Engine/Plugins/Experimental/AVCodecs/AVCodecsCore/Source/AVCodecsCore/Public/Video/Encoders/SVC/ScalabilityStructureFull.h
@@ -0,0 +1,198 @@
+// Copyright Epic Games, Inc. All Rights Reserved.
+
+#pragma once
+
+#include "Video/GenericFrameInfo.h"
+#include "Video/Encoders/SVC/ScalableVideoController.h"
+
+class AVCODECSCORE_API FScalabilityStructureFullSvc : public FScalableVideoController
+{
+public:
+	FScalabilityStructureFullSvc(int NumSpatialLayers, int NumTemporalLayers, FIntFraction ResolutionFactor);
+	~FScalabilityStructureFullSvc() override = default;
+
+	virtual FStreamLayersConfig StreamConfig() const override;
+
+	virtual TArray<FScalableVideoController::FLayerFrameConfig> NextFrameConfig(bool bRestart) override;
+	virtual FGenericFrameInfo									OnEncodeDone(const FScalableVideoController::FLayerFrameConfig& Config) override;
+	virtual void												OnRatesUpdated(const FVideoBitrateAllocation& Bitrates) override;
+
+private:
+	// NOTE: This enum name is duplicated throughout the ScalabilityStructure variants
+	// While it's name is the same, the order of the values is imporant and differs between variants
+	enum EFramePattern : uint8
+	{
+		None,
+		Key,
+		DeltaT2A,
+		DeltaT1,
+		DeltaT2B,
+		DeltaT0,
+	};
+
+	static FString		 FramePatternNames[6];
+	static constexpr int MaxNumSpatialLayers = 3;
+	static constexpr int MaxNumTemporalLayers = 3;
+
+	// Index of the buffer to store last frame for layer (`Sid`, `Tid`)
+	int BufferIndex(int Sid, int Tid) const
+	{
+		return Tid * NumSpatialLayers + Sid;
+	}
+
+	bool DecodeTargetIsActive(int Sid, int Tid) const
+	{
+		return static_cast<bool>(ActiveDecodeTargets[Sid * NumTemporalLayers + Tid]);
+	}
+
+	void SetDecodeTargetIsActive(int Sid, int Tid, bool Value)
+	{
+		ActiveDecodeTargets[Sid * NumTemporalLayers + Tid] = Value;
+	}
+
+	EFramePattern				   NextPattern() const;
+	bool						   TemporalLayerIsActive(int Tid) const;
+	static EDecodeTargetIndication Dti(int Sid, int Tid, const FScalableVideoController::FLayerFrameConfig& Frame);
+
+	const int			 NumSpatialLayers;
+	const int			 NumTemporalLayers;
+	const FIntFraction	 ResolutionFactor;
+
+	EFramePattern LastPattern = EFramePattern::None;
+	TBitArray<>	  CanReferenceT0FrameForSpatialId;
+	TBitArray<>	  CanReferenceT1FrameForSpatialId;
+	TArray<bool>  ActiveDecodeTargets;
+};
+
+// T1       0   0
+//         /   /   / ...
+// T0     0---0---0--
+// Time-> 0 1 2 3 4
+class AVCODECSCORE_API FScalabilityStructureL1T2 : public FScalabilityStructureFullSvc
+{
+public:
+	explicit FScalabilityStructureL1T2(FIntFraction ResolutionFactor = { 1, 2 })
+		: FScalabilityStructureFullSvc(1, 2, ResolutionFactor)
+	{
+	}
+	~FScalabilityStructureL1T2() override = default;
+
+	FFrameDependencyStructure DependencyStructure() const override;
+};
+
+// T2       0   0   0   0
+//          |  /    |  /
+// T1       / 0     / 0  ...
+//         |_/     |_/
+// T0     0-------0------
+// Time-> 0 1 2 3 4 5 6 7
+class AVCODECSCORE_API FScalabilityStructureL1T3 : public FScalabilityStructureFullSvc
+{
+public:
+	explicit FScalabilityStructureL1T3(FIntFraction ResolutionFactor = { 1, 2 })
+		: FScalabilityStructureFullSvc(1, 3, ResolutionFactor)
+	{
+	}
+	~FScalabilityStructureL1T3() override = default;
+
+	FFrameDependencyStructure DependencyStructure() const override;
+};
+
+// S1  0--0--0-
+//     |  |  | ...
+// S0  0--0--0-
+class AVCODECSCORE_API FScalabilityStructureL2T1 : public FScalabilityStructureFullSvc
+{
+public:
+	explicit FScalabilityStructureL2T1(FIntFraction ResolutionFactor = { 1, 2 })
+		: FScalabilityStructureFullSvc(2, 1, ResolutionFactor)
+	{
+	}
+	~FScalabilityStructureL2T1() override = default;
+
+	FFrameDependencyStructure DependencyStructure() const override;
+};
+
+// S1T1     0   0
+//         /|  /|  /
+// S1T0   0-+-0-+-0
+//        | | | | | ...
+// S0T1   | 0 | 0 |
+//        |/  |/  |/
+// S0T0   0---0---0--
+// Time-> 0 1 2 3 4
+class AVCODECSCORE_API FScalabilityStructureL2T2 : public FScalabilityStructureFullSvc
+{
+public:
+	explicit FScalabilityStructureL2T2(FIntFraction ResolutionFactor = { 1, 2 })
+		: FScalabilityStructureFullSvc(2, 2, ResolutionFactor)
+	{
+	}
+	~FScalabilityStructureL2T2() override = default;
+
+	FFrameDependencyStructure DependencyStructure() const override;
+};
+
+// S1T2      4    ,8
+// S1T1    / |  6' |
+// S1T0   2--+-'+--+-...
+//        |  |  |  |
+// S0T2   |  3  | ,7
+// S0T1   | /   5'
+// S0T0   1----'-----...
+// Time-> 0  1  2  3
+class AVCODECSCORE_API FScalabilityStructureL2T3 : public FScalabilityStructureFullSvc
+{
+public:
+	explicit FScalabilityStructureL2T3(FIntFraction ResolutionFactor = { 1, 2 })
+		: FScalabilityStructureFullSvc(2, 3, ResolutionFactor)
+	{
+	}
+	~FScalabilityStructureL2T3() override = default;
+
+	FFrameDependencyStructure DependencyStructure() const override;
+};
+
+// S2     0-0-0-
+//        | | |
+// S1     0-0-0-...
+//        | | |
+// S0     0-0-0-
+// Time-> 0 1 2
+class AVCODECSCORE_API FScalabilityStructureL3T1 : public FScalabilityStructureFullSvc
+{
+public:
+	explicit FScalabilityStructureL3T1(FIntFraction ResolutionFactor = { 1, 2 })
+		: FScalabilityStructureFullSvc(3, 1, ResolutionFactor)
+	{
+	}
+	~FScalabilityStructureL3T1() override = default;
+
+	FFrameDependencyStructure DependencyStructure() const override;
+};
+
+// https://www.w3.org/TR/webrtc-svc/#L3T2*
+class AVCODECSCORE_API FScalabilityStructureL3T2 : public FScalabilityStructureFullSvc
+{
+public:
+	explicit FScalabilityStructureL3T2(FIntFraction ResolutionFactor = { 1, 2 })
+		: FScalabilityStructureFullSvc(3, 2, ResolutionFactor)
+	{
+	}
+	~FScalabilityStructureL3T2() override = default;
+
+	FFrameDependencyStructure DependencyStructure() const override;
+};
+
+// https://www.w3.org/TR/webrtc-svc/#L3T3*
+class AVCODECSCORE_API FScalabilityStructureL3T3 : public FScalabilityStructureFullSvc
+{
+public:
+	explicit FScalabilityStructureL3T3(FIntFraction ResolutionFactor = { 1, 2 })
+		: FScalabilityStructureFullSvc(3, 3, ResolutionFactor)
+	{
+	}
+	~FScalabilityStructureL3T3() override = default;
+
+	FFrameDependencyStructure DependencyStructure() const override;
+};
\ No newline at end of file
diff --git a/Engine/Plugins/Experimental/AVCodecs/AVCodecsCore/Source/AVCodecsCore/Public/Video/Encoders/SVC/ScalabilityStructureKey.h b/Engine/Plugins/Experimental/AVCodecs/AVCodecsCore/Source/AVCodecsCore/Public/Video/Encoders/SVC/ScalabilityStructureKey.h
new file mode 100644
--- /dev/null
+++ b/Engine/Plugins/Experimental/AVCodecs/AVCodecsCore/Source/AVCodecsCore/Public/Video/Encoders/SVC/ScalabilityStructureKey.h
@@ -0,0 +1,152 @@
+// Copyright Epic Games, Inc. All Rights Reserved.
+
+#pragma once
+
+#include "Video/GenericFrameInfo.h"
+#include "Video/Encoders/SVC/ScalableVideoController.h"
+
+class AVCODECSCORE_API FScalabilityStructureKeySvc : public FScalableVideoController
+{
+public:
+	FScalabilityStructureKeySvc(int NumSpatialLayers, int NumTemporalLayers);
+	~FScalabilityStructureKeySvc() override = default;
+
+	virtual FStreamLayersConfig StreamConfig() const override;
+
+	virtual TArray<FScalableVideoController::FLayerFrameConfig> NextFrameConfig(bool bRestart) override;
+	virtual FGenericFrameInfo									OnEncodeDone(const FScalableVideoController::FLayerFrameConfig& Config) override;
+	virtual void												OnRatesUpdated(const FVideoBitrateAllocation& Bitrates) override;
+
+private:
+	// NOTE: This enum name is duplicated throughout the ScalabilityStructure variants
+	// While it's name is the same, the order of the values is imporant and differs between variants
+	enum EFramePattern : uint8
+	{
+		None,
+		Key,
+		DeltaT0,
+		DeltaT2A,
+		DeltaT1,
+		DeltaT2B,
+	};
+
+	static constexpr int MaxNumSpatialLayers = 3;
+	static constexpr int MaxNumTemporalLayers = 3;
+
+	// Index of the buffer to store last frame for layer (`Sid`, `Tid`)
+	int BufferIndex(int Sid, int Tid) const
+	{
+		return Tid * NumSpatialLayers + Sid;
+	}
+
+	bool DecodeTargetIsActive(int Sid, int Tid) const
+	{
+		return static_cast<bool>(ActiveDecodeTargets[Sid * NumTemporalLayers + Tid]);
+	}
+
+	void SetDecodeTargetIsActive(int Sid, int Tid, bool value)
+	{
+		ActiveDecodeTargets[Sid * NumTemporalLayers + Tid] = value;
+	}
+
+	bool						   TemporalLayerIsActive(int Tid) const;
+	static EDecodeTargetIndication Dti(int Sid, int Tid, const FScalableVideoController::FLayerFrameConfig& config);
+
+	TArray<FScalableVideoController::FLayerFrameConfig> KeyframeConfig();
+	TArray<FScalableVideoController::FLayerFrameConfig> T0Config();
+	TArray<FScalableVideoController::FLayerFrameConfig> T1Config();
+	TArray<FScalableVideoController::FLayerFrameConfig> T2Config(EFramePattern Pattern);
+
+	EFramePattern NextPattern(EFramePattern InLastPattern) const;
+
+	const int NumSpatialLayers;
+	const int NumTemporalLayers;
+
+	EFramePattern LastPattern = EFramePattern::None;
+	TBitArray<>	  SpatialIdIsEnabled;
+	TBitArray<>	  CanReferenceT1FrameForSpatialId;
+	TArray<bool>  ActiveDecodeTargets;
+};
+
+// S1  0--0--0-
+//     |       ...
+// S0  0--0--0-
+class AVCODECSCORE_API FScalabilityStructureL2T1Key : public FScalabilityStructureKeySvc
+{
+public:
+	FScalabilityStructureL2T1Key()
+		: FScalabilityStructureKeySvc(2, 1)
+	{
+	}
+	~FScalabilityStructureL2T1Key() override = default;
+
+	FFrameDependencyStructure DependencyStructure() const override;
+};
+
+// S1T1     0   0
+//         /   /   /
+// S1T0   0---0---0
+//        |         ...
+// S0T1   | 0   0
+//        |/   /   /
+// S0T0   0---0---0
+// Time-> 0 1 2 3 4
+class AVCODECSCORE_API FScalabilityStructureL2T2Key : public FScalabilityStructureKeySvc
+{
+public:
+	FScalabilityStructureL2T2Key()
+		: FScalabilityStructureKeySvc(2, 2)
+	{
+	}
+	~FScalabilityStructureL2T2Key() override = default;
+
+	FFrameDependencyStructure DependencyStructure() const override;
+};
+
+class AVCODECSCORE_API FScalabilityStructureL2T3Key : public FScalabilityStructureKeySvc
+{
+public:
+	FScalabilityStructureL2T3Key()
+		: FScalabilityStructureKeySvc(2, 3)
+	{
+	}
+	~FScalabilityStructureL2T3Key() override = default;
+
+	FFrameDependencyStructure DependencyStructure() const override;
+};
+
+class AVCODECSCORE_API FScalabilityStructureL3T1Key : public FScalabilityStructureKeySvc
+{
+public:
+	FScalabilityStructureL3T1Key()
+		: FScalabilityStructureKeySvc(3, 1)
+	{
+	}
+	~FScalabilityStructureL3T1Key() override = default;
+
+	FFrameDependencyStructure DependencyStructure() const override;
+};
+
+class AVCODECSCORE_API FScalabilityStructureL3T2Key : public FScalabilityStructureKeySvc
+{
+public:
+	FScalabilityStructureL3T2Key()
+		: FScalabilityStructureKeySvc(3, 2)
+	{
+	}
+	~FScalabilityStructureL3T2Key() override = default;
+
+	FFrameDependencyStructure DependencyStructure() const override;
+};
+
+class AVCODECSCORE_API FScalabilityStructureL3T3Key : public FScalabilityStructureKeySvc
+{
+public:
+	FScalabilityStructureL3T3Key()
+		: FScalabilityStructureKeySvc(3, 3)
+	{
+	}
+	~FScalabilityStructureL3T3Key() override = default;
+
+	FFrameDependencyStructure DependencyStructure() const override;
+};
\ No newline at end of file
diff --git a/Engine/Plugins/Experimental/AVCodecs/AVCodecsCore/Source/AVCodecsCore/Public/Video/Encoders/SVC/ScalabilityStructureL2T2KeyShift.h b/Engine/Plugins/Experimental/AVCodecs/AVCodecsCore/Source/AVCodecsCore/Public/Video/Encoders/SVC/ScalabilityStructureL2T2KeyShift.h
new file mode 100644
--- /dev/null
+++ b/Engine/Plugins/Experimental/AVCodecs/AVCodecsCore/Source/AVCodecsCore/Public/Video/Encoders/SVC/ScalabilityStructureL2T2KeyShift.h
@@ -0,0 +1,54 @@
+// Copyright Epic Games, Inc. All Rights Reserved.
+
+#pragma once
+
+#include "Video/GenericFrameInfo.h"
+#include "Video/Encoders/SVC/ScalableVideoController.h"
+
+// S1T1     0   0
+//         /   /   /
+// S1T0   0---0---0
+//        |        ...
+// S0T1   |   0   0
+//        |  /   /
+// S0T0   0-0---0--
+// Time-> 0 1 2 3 4
+class AVCODECSCORE_API FScalabilityStructureL2T2KeyShift : public FScalableVideoController
+{
+public:
+	FScalabilityStructureL2T2KeyShift();
+	~FScalabilityStructureL2T2KeyShift() override = default;
+
+	virtual FStreamLayersConfig		  StreamConfig() const override;
+	virtual FFrameDependencyStructure DependencyStructure() const override;
+
+	virtual TArray<FScalableVideoController::FLayerFrameConfig> NextFrameConfig(bool bRestart) override;
+	virtual FGenericFrameInfo									OnEncodeDone(const FScalableVideoController::FLayerFrameConfig& Config) override;
+	virtual void												OnRatesUpdated(const FVideoBitrateAllocation& Bitrates) override;
+
+private:
+	// NOTE: This enum name is duplicated throughout the ScalabilityStructure variants
+	// While it's name is the same, the order of the values is imporant and differs between variants
+	enum EFramePattern : uint8
+	{
+		Key,
+		Delta0,
+		Delta1,
+	};
+
+	static constexpr int NumSpatialLayers = 2;
+	static constexpr int NumTemporalLayers = 2;
+
+	bool DecodeTargetIsActive(int Sid, int Tid) const
+	{
+		return static_cast<bool>(ActiveDecodeTargets[Sid * NumTemporalLayers + Tid]);
+	}
+	void SetDecodeTargetIsActive(int Sid, int Tid, bool value)
+	{
+		ActiveDecodeTargets[Sid * NumTemporalLayers + Tid] = value;
+	}
+	static EDecodeTargetIndication Dti(int Sid, int Tid, const FScalableVideoController::FLayerFrameConfig& Config);
+
+	EFramePattern NextPattern = EFramePattern::Key;
+	TArray<bool>  ActiveDecodeTargets;
+};
\ No newline at end of file
diff --git a/Engine/Plugins/Experimental/AVCodecs/AVCodecsCore/Source/AVCodecsCore/Public/Video/Encoders/SVC/ScalabilityStructureSimulcast.h b/Engine/Plugins/Experimental/AVCodecs/AVCodecsCore/Source/AVCodecsCore/Public/Video/Encoders/SVC/ScalabilityStructureSimulcast.h
new file mode 100644
--- /dev/null
+++ b/Engine/Plugins/Experimental/AVCodecs/AVCodecsCore/Source/AVCodecsCore/Public/Video/Encoders/SVC/ScalabilityStructureSimulcast.h
@@ -0,0 +1,146 @@
+// Copyright Epic Games, Inc. All Rights Reserved.
+
+#pragma once
+
+#include "Video/GenericFrameInfo.h"
+#include "Video/Encoders/SVC/ScalableVideoController.h"
+
+// Scalability structure with multiple independent spatial layers each with the
+// same temporal layering.
+class AVCODECSCORE_API FScalabilityStructureSimulcast : public FScalableVideoController
+{
+public:
+	FScalabilityStructureSimulcast(int NumSpatialLayers, int NumTemporalLayers, FIntFraction ResolutionFactor);
+	virtual ~FScalabilityStructureSimulcast() override = default;
+
+	virtual FStreamLayersConfig									StreamConfig() const override;
+	virtual TArray<FScalableVideoController::FLayerFrameConfig> NextFrameConfig(bool bRestart) override;
+	virtual FGenericFrameInfo									OnEncodeDone(const FScalableVideoController::FLayerFrameConfig& Config) override;
+	virtual void												OnRatesUpdated(const FVideoBitrateAllocation& Bitrates) override;
+
+private:
+	// NOTE: This enum name is duplicated throughout the ScalabilityStructure variants
+	// While it's name is the same, the order of the values is imporant and differs between variants
+	enum EFramePattern : uint8
+	{
+		None,
+		DeltaT2A,
+		DeltaT1,
+		DeltaT2B,
+		DeltaT0,
+	};
+	static constexpr int MaxNumSpatialLayers = 3;
+	static constexpr int MaxNumTemporalLayers = 3;
+
+	// Index of the buffer to store last frame for layer (`Sid`, `Tid`)
+	int BufferIndex(int Sid, int Tid) const
+	{
+		return Tid * NumSpatialLayers + Sid;
+	}
+
+	bool DecodeTargetIsActive(int Sid, int Tid) const
+	{
+		return static_cast<bool>(ActiveDecodeTargets[Sid * NumTemporalLayers + Tid]);
+	}
+
+	void SetDecodeTargetIsActive(int Sid, int Tid, bool Value)
+	{
+		ActiveDecodeTargets[Sid * NumTemporalLayers + Tid] = Value;
+	}
+
+	EFramePattern				   NextPattern() const;
+	bool						   TemporalLayerIsActive(int Tid) const;
+	static EDecodeTargetIndication Dti(int Sid, int Tid, const FScalableVideoController::FLayerFrameConfig& Config);
+
+	const int	   NumSpatialLayers;
+	const int	   NumTemporalLayers;
+	FIntFraction   ResolutionFactor;
+
+	EFramePattern LastPattern = EFramePattern::None;
+	TBitArray<>	  CanReferenceT0FrameForSpatialId;
+	TBitArray<>	  CanReferenceT1FrameForSpatialId;
+	TArray<bool>  ActiveDecodeTargets;
+};
+
+// S1  0--0--0-
+//             ...
+// S0  0--0--0-
+class AVCODECSCORE_API FScalabilityStructureS2T1 : public FScalabilityStructureSimulcast
+{
+public:
+	FScalabilityStructureS2T1(FIntFraction ResolutionFactor = { 1, 2 })
+		: FScalabilityStructureSimulcast(2, 1, ResolutionFactor)
+	{
+	}
+	~FScalabilityStructureS2T1() override = default;
+
+	FFrameDependencyStructure DependencyStructure() const override;
+};
+
+class AVCODECSCORE_API FScalabilityStructureS2T2 : public FScalabilityStructureSimulcast
+{
+public:
+	explicit FScalabilityStructureS2T2(FIntFraction ResolutionFactor = { 1, 2 })
+		: FScalabilityStructureSimulcast(2, 2, ResolutionFactor) {}
+	~FScalabilityStructureS2T2() override = default;
+
+	FFrameDependencyStructure DependencyStructure() const override;
+};
+
+// S1T2       3   7
+//            |  /
+// S1T1       / 5
+//           |_/
+// S1T0     1-------9...
+//
+// S0T2       2   6
+//            |  /
+// S0T1       / 4
+//           |_/
+// S0T0     0-------8...
+// Time->   0 1 2 3 4
+class AVCODECSCORE_API FScalabilityStructureS2T3 : public FScalabilityStructureSimulcast
+{
+public:
+	explicit FScalabilityStructureS2T3(FIntFraction ResolutionFactor = { 1, 2 })
+		: FScalabilityStructureSimulcast(2, 3, ResolutionFactor)
+	{
+	}
+	~FScalabilityStructureS2T3() override = default;
+
+	FFrameDependencyStructure DependencyStructure() const override;
+};
+
+class AVCODECSCORE_API FScalabilityStructureS3T1 : public FScalabilityStructureSimulcast
+{
+public:
+	explicit FScalabilityStructureS3T1(FIntFraction ResolutionFactor = { 1, 2 })
+		: FScalabilityStructureSimulcast(3, 1, ResolutionFactor) {}
+	~FScalabilityStructureS3T1() override = default;
+
+	FFrameDependencyStructure DependencyStructure() const override;
+};
+
+class AVCODECSCORE_API FScalabilityStructureS3T2 : public FScalabilityStructureSimulcast
+{
+public:
+	explicit FScalabilityStructureS3T2(FIntFraction ResolutionFactor = { 1, 2 })
+		: FScalabilityStructureSimulcast(3, 2, ResolutionFactor)
+	{
+	}
+	~FScalabilityStructureS3T2() override = default;
+
+	FFrameDependencyStructure DependencyStructure() const override;
+};
+
+class AVCODECSCORE_API FScalabilityStructureS3T3 : public FScalabilityStructureSimulcast
+{
+public:
+	FScalabilityStructureS3T3(FIntFraction ResolutionFactor = { 1, 2 })
+		: FScalabilityStructureSimulcast(3, 3, ResolutionFactor)
+	{
+	}
+	~FScalabilityStructureS3T3() override = default;
+
+	FFrameDependencyStructure DependencyStructure() const override;
+};
\ No newline at end of file
diff --git a/Engine/Plugins/Experimental/AVCodecs/AVCodecsCore/Source/AVCodecsCore/Public/Video/Encoders/SVC/ScalableVideoControllerNoLayering.h b/Engine/Plugins/Experimental/AVCodecs/AVCodecsCore/Source/AVCodecsCore/Public/Video/Encoders/SVC/ScalableVideoControllerNoLayering.h
new file mode 100644
--- /dev/null
+++ b/Engine/Plugins/Experimental/AVCodecs/AVCodecsCore/Source/AVCodecsCore/Public/Video/Encoders/SVC/ScalableVideoControllerNoLayering.h
@@ -0,0 +1,22 @@
+// Copyright Epic Games, Inc. All Rights Reserved.
+
+#pragma once
+
+#include "Video/Encoders/SVC/ScalableVideoController.h"
+
+class AVCODECSCORE_API FScalableVideoControllerNoLayering : public FScalableVideoController
+{
+public:
+	~FScalableVideoControllerNoLayering() override = default;
+
+	virtual FStreamLayersConfig		  StreamConfig() const override;
+	virtual FFrameDependencyStructure DependencyStructure() const override;
+
+	virtual TArray<FScalableVideoController::FLayerFrameConfig> NextFrameConfig(bool bRestart) override;
+	virtual FGenericFrameInfo									OnEncodeDone(const FScalableVideoController::FLayerFrameConfig& Config) override;
+	virtual void												OnRatesUpdated(const FVideoBitrateAllocation& Bitrates) override;
+
+private:
+	bool bStart = true;
+	bool bEnabled = true;
+};
\ No newline at end of file
diff --git a/Engine/Plugins/Experimental/AVCodecs/AVCodecsCore/Source/AVCodecsCore/Public/Video/Encoders/SVC/VideoBitrateAllocatorSVC.h b/Engine/Plugins/Experimental/AVCodecs/AVCodecsCore/Source/AVCodecsCore/Public/Video/Encoders/SVC/VideoBitrateAllocatorSVC.h
new file mode 100644
--- /dev/null
+++ b/Engine/Plugins/Experimental/AVCodecs/AVCodecsCore/Source/AVCodecsCore/Public/Video/Encoders/SVC/VideoBitrateAllocatorSVC.h
@@ -0,0 +1,39 @@
+// Copyright Epic Games, Inc. All Rights Reserved.
+
+#pragma once
+
+#include "Video/Encoders/VideoBitrateAllocator.h"
+#include "Video/VideoEncoder.h"
+
+class AVCODECSCORE_API FVideoBitrateAllocatorSVC : public FVideoBitrateAllocator
+{
+public:
+	// NOTE: Config stores its rates as kbps whereas this class will use bps. Convert to bps everywhere
+	explicit FVideoBitrateAllocatorSVC(const FVideoEncoderConfig& Config);
+	virtual ~FVideoBitrateAllocatorSVC() = default;
+
+	FVideoBitrateAllocation Allocate(FVideoBitrateAllocationParameters Parameters) override;
+
+	static uint32		  GetMaxBitrate(const FVideoEncoderConfig& Config);
+	static uint32		  GetPaddingBitrate(const FVideoEncoderConfig& Config);
+	static TArray<uint32> GetLayerStartBitrates(const FVideoEncoderConfig& Config);
+
+private:
+	struct FNumLayers
+	{
+		size_t Spatial = 1;
+		size_t Temporal = 1;
+	};
+
+	static FNumLayers		GetNumLayers(const FVideoEncoderConfig& Config);
+	FVideoBitrateAllocation GetAllocationNormalVideo(uint32 TotalBitrateBps, size_t FirstActiveLayer, size_t NumSpatialLayers) const;
+
+	// Returns the number of layers that are active and have enough bitrate to
+	// actually be enabled.
+	size_t FindNumEnabledLayers(uint32 TargetRate) const;
+
+	FVideoEncoderConfig	 Config;
+	const FNumLayers	 NumLayers;
+	const TArray<uint32> CumulativeLayerStartBitrates;
+	size_t				 LastActiveLayerCount;
+};
\ No newline at end of file
diff --git a/Engine/Plugins/Experimental/AVCodecs/AVCodecsCore/Source/AVCodecsCore/Public/Video/Encoders/VideoBitrateAllocation.h b/Engine/Plugins/Experimental/AVCodecs/AVCodecsCore/Source/AVCodecsCore/Public/Video/Encoders/VideoBitrateAllocation.h
new file mode 100644
--- /dev/null
+++ b/Engine/Plugins/Experimental/AVCodecs/AVCodecsCore/Source/AVCodecsCore/Public/Video/Encoders/VideoBitrateAllocation.h
@@ -0,0 +1,70 @@
+// Copyright Epic Games, Inc. All Rights Reserved.
+
+#pragma once
+
+#include "Containers/Array.h"
+#include "Math/NumericLimits.h"
+#include "Misc/Optional.h"
+#include "Video/VideoEncoder.h"
+
+class AVCODECSCORE_API FVideoBitrateAllocation
+{
+public:
+	static constexpr uint32 MaxBitrateBps = TNumericLimits<uint32>::Max();
+
+	FVideoBitrateAllocation();
+
+	virtual ~FVideoBitrateAllocation() = default;
+
+	bool HasBitrate(uint64 SpatialIndex, uint64 TemporalIndex) const;
+
+	uint32 GetBitrate(uint64 SpatialIndex, uint64 TemporalIndex) const;
+
+	bool IsSpatialLayerUsed(uint64 SpatialIndex) const;
+
+	uint32 GetSpatialLayerSumBitrate(uint64 SpatialIndex) const;
+
+	uint32 GetTemporalLayerSumBitrate(uint64 SpatialIndex, uint64 TemporalIndex) const;
+
+	TArray<uint32> GetTemporalLayerAllocation(uint64 SpatialIndex) const;
+
+	uint32 GetSumBps() const;
+
+	bool IsBwLimited() const;
+
+	bool SetBitrate(size_t SpatialIndex, size_t TemporalIndex, uint32 BitrateBps);
+
+	// Indicates if the allocation has some layers/streams disabled due to
+	// low available bandwidth.
+	void SetBwLimited(bool Limited);
+
+	// Returns one FVideoBitrateAllocation for each spatial layer. This is used to
+	// configure simulcast streams. Note that the length of the returned vector is
+	// always 5, the optional is unset for unused layers.
+	TArray<TOptional<FVideoBitrateAllocation>> GetSimulcastAllocations() const;
+
+	inline bool operator==(const FVideoBitrateAllocation& Other) const
+	{
+		for (size_t si = 0; si < 5; ++si)
+		{
+			for (size_t ti = 0; ti < 4; ++ti)
+			{
+				if (Bitrates[si][ti] != Other.Bitrates[si][ti])
+				{
+					return false;
+				}
+			}
+		}
+		return true;
+	}
+
+	inline bool operator!=(const FVideoBitrateAllocation& Other) const
+	{
+		return !(*this == Other);
+	}
+
+private:
+	uint32			  SumBps;
+	TOptional<uint32> Bitrates[Video::MaxSpatialLayers][Video::MaxTemporalStreams];
+	bool			  bIsBwLimited;
+};
\ No newline at end of file
diff --git a/Engine/Plugins/Experimental/AVCodecs/AVCodecsCore/Source/AVCodecsCore/Public/Video/Encoders/VideoBitrateAllocator.h b/Engine/Plugins/Experimental/AVCodecs/AVCodecsCore/Source/AVCodecsCore/Public/Video/Encoders/VideoBitrateAllocator.h
new file mode 100644
--- /dev/null
+++ b/Engine/Plugins/Experimental/AVCodecs/AVCodecsCore/Source/AVCodecsCore/Public/Video/Encoders/VideoBitrateAllocator.h
@@ -0,0 +1,38 @@
+// Copyright Epic Games, Inc. All Rights Reserved.
+
+#pragma once
+
+#include "Misc/FrameRate.h"
+#include "Video/Encoders/VideoBitrateAllocation.h"
+
+class AVCODECSCORE_API FVideoBitrateAllocationParameters
+{
+public:
+	FVideoBitrateAllocationParameters(uint32 TotalBitrateBps, FFrameRate Framerate)
+		: TotalBitrateBps(TotalBitrateBps)
+		, StableBitrateBps(TotalBitrateBps)
+		, Framerate(Framerate)
+	{
+	}
+	FVideoBitrateAllocationParameters(uint32 TotalBitrateBps, uint32 StableBitrateBps, FFrameRate Framerate)
+		: TotalBitrateBps(TotalBitrateBps)
+		, StableBitrateBps(StableBitrateBps)
+		, Framerate(Framerate)
+	{
+	}
+
+	uint32 TotalBitrateBps;
+	uint32 StableBitrateBps;
+	FFrameRate Framerate;
+};
+
+class AVCODECSCORE_API FVideoBitrateAllocator
+{
+public:
+	FVideoBitrateAllocator() = default;
+	~FVideoBitrateAllocator() = default;
+
+	virtual FVideoBitrateAllocation GetAllocation(uint32 TotalBitrateBps, FFrameRate Framerate);
+
+	virtual FVideoBitrateAllocation Allocate(FVideoBitrateAllocationParameters Parameters);
+};
\ No newline at end of file
diff --git a/Engine/Plugins/Experimental/AVCodecs/AVCodecsCore/Source/AVCodecsCore/Public/Video/Resources/VideoResourceCPU.h b/Engine/Plugins/Experimental/AVCodecs/AVCodecsCore/Source/AVCodecsCore/Public/Video/Resources/VideoResourceCPU.h
--- a/Engine/Plugins/Experimental/AVCodecs/AVCodecsCore/Source/AVCodecsCore/Public/Video/Resources/VideoResourceCPU.h
+++ b/Engine/Plugins/Experimental/AVCodecs/AVCodecsCore/Source/AVCodecsCore/Public/Video/Resources/VideoResourceCPU.h
@@ -22,6 +22,7 @@ private:
 
 public:
 	FORCEINLINE TSharedPtr<uint8> const& GetRaw() const { return Raw; }
+	FORCEINLINE void					 SetRaw(TSharedPtr<uint8> RawData) { Raw = RawData; }
 
 	FVideoResourceCPU(TSharedRef<FAVDevice> const& Device, TSharedPtr<uint8> const& Raw, FAVLayout const& Layout, FVideoDescriptor const& Descriptor);
 	virtual ~FVideoResourceCPU() override = default;
@@ -29,5 +30,16 @@ public:
 	virtual FAVResult Validate() const override;
 };
 
+class AVCODECSCORE_API FResolvableVideoResourceCPU : public TResolvableVideoResource<FVideoResourceCPU>
+{
+protected:
+	virtual TSharedPtr<FVideoResourceCPU> TryResolve(TSharedPtr<FAVDevice> const& Device, FVideoDescriptor const& Descriptor) override;
+
+private:
+	uint32 GetStrideInterleavedOrLuma(FVideoDescriptor const& Descriptor);
+	uint32 GetSize(FVideoDescriptor const& Descriptor);
+};
+
 DECLARE_TYPEID(FVideoContextCPU, AVCODECSCORE_API);
 DECLARE_TYPEID(FVideoResourceCPU, AVCODECSCORE_API);
+DECLARE_TYPEID(FResolvableVideoResourceCPU, AVCODECSCORE_API);
\ No newline at end of file
diff --git a/Engine/Plugins/Experimental/AVCodecs/AVCodecsCore/Source/AVCodecsCore/Public/Video/VideoDecoder.h b/Engine/Plugins/Experimental/AVCodecs/AVCodecsCore/Source/AVCodecsCore/Public/Video/VideoDecoder.h
--- a/Engine/Plugins/Experimental/AVCodecs/AVCodecsCore/Source/AVCodecsCore/Public/Video/VideoDecoder.h
+++ b/Engine/Plugins/Experimental/AVCodecs/AVCodecsCore/Source/AVCodecsCore/Public/Video/VideoDecoder.h
@@ -3,6 +3,7 @@
 #pragma once
 
 #include "AVCoder.h"
+#include "Video/VideoConfig.h"
 #include "Video/VideoPacket.h"
 #include "Video/VideoResource.h"
 
@@ -10,11 +11,11 @@
  * Implementation of Video Decoding domain, see TAVCoder for inheritance model
  */
 
-struct FVideoDecoderConfig : public FAVConfig
+struct FVideoDecoderConfig : public FVideoConfig
 {
 public:
 	FVideoDecoderConfig(EAVPreset Preset = EAVPreset::Default)
-		: FAVConfig(Preset) { }
+		: FVideoConfig(Preset) {}
 };
 
 /**
@@ -64,8 +65,7 @@ public:
 				return FAVResult(EAVResult::ErrorInvalidState, TEXT("Decoder not open"));
 			}
 
-			TDelegatedVideoResource<TChildResource> WrappedResource([&InOutResource](TSharedPtr<TChildResource>& InOutResult, TSharedPtr<FAVDevice> const& InDevice, FVideoDescriptor const& Descriptor)
-			{
+			TDelegatedVideoResource<TChildResource> WrappedResource([&InOutResource](TSharedPtr<TChildResource>& InOutResult, TSharedPtr<FAVDevice> const& InDevice, FVideoDescriptor const& Descriptor) {
 				InOutResult = nullptr;
 
 				if (InOutResource.Resolve(InDevice, Descriptor))
@@ -146,7 +146,7 @@ public:
 
 	/**
 	 * Flush remaining frames and invalidate the underlying architecture.
-	 * 
+	 *
 	 * @return Result of the operation, @see FAVResult.
 	 */
 	FAVResult FlushFrames()
diff --git a/Engine/Plugins/Experimental/AVCodecs/AVCodecsCore/Source/AVCodecsCore/Public/Video/VideoEncoder.h b/Engine/Plugins/Experimental/AVCodecs/AVCodecsCore/Source/AVCodecsCore/Public/Video/VideoEncoder.h
--- a/Engine/Plugins/Experimental/AVCodecs/AVCodecsCore/Source/AVCodecsCore/Public/Video/VideoEncoder.h
+++ b/Engine/Plugins/Experimental/AVCodecs/AVCodecsCore/Source/AVCodecsCore/Public/Video/VideoEncoder.h
@@ -3,12 +3,100 @@
 #pragma once
 
 #include "AVCoder.h"
+#include "AVConstants.h"
+#include "Video/VideoConfig.h"
 #include "Video/VideoPacket.h"
 #include "Video/VideoResource.h"
 
 /*
  * Implementation of Video Encoding domain, see TAVCoder for inheritance model
  */
+enum class EScalabilityMode : uint8
+{
+	L1T1 = 0,
+	L1T2,
+	L1T3,
+	L2T1,
+	L2T1h,
+	L2T1_KEY,
+	L2T2,
+	L2T2h,
+	L2T2_KEY,
+	L2T2_KEY_SHIFT,
+	L2T3,
+	L2T3h,
+	L2T3_KEY,
+	L3T1,
+	L3T1h,
+	L3T1_KEY,
+	L3T2,
+	L3T2h,
+	L3T2_KEY,
+	L3T3,
+	L3T3h,
+	L3T3_KEY,
+	S2T1,
+	S2T1h,
+	S2T2,
+	S2T2h,
+	S2T3,
+	S2T3h,
+	S3T1,
+	S3T1h,
+	S3T2,
+	S3T2h,
+	S3T3,
+	S3T3h,
+	None
+};
+
+TMap<FString, EScalabilityMode> const ScalabilityModeMap = {
+	{ "L1T1", EScalabilityMode::L1T1 },
+	{ "L1T2", EScalabilityMode::L1T2 },
+	{ "L1T3", EScalabilityMode::L1T3 },
+	{ "L2T1", EScalabilityMode::L2T1 },
+	{ "L2T1h", EScalabilityMode::L2T1h },
+	{ "L2T1_KEY", EScalabilityMode::L2T1_KEY },
+	{ "L2T2", EScalabilityMode::L2T2 },
+	{ "L2T2h", EScalabilityMode::L2T2h },
+	{ "L2T2_KEY", EScalabilityMode::L2T2_KEY },
+	{ "L2T2_KEY_SHIFT", EScalabilityMode::L2T2_KEY_SHIFT },
+	{ "L2T3", EScalabilityMode::L2T3 },
+	{ "L2T3h", EScalabilityMode::L2T3h },
+	{ "L2T3_KEY", EScalabilityMode::L2T3_KEY },
+	{ "L3T1", EScalabilityMode::L3T1 },
+	{ "L3T1h", EScalabilityMode::L3T1h },
+	{ "L3T1_KEY", EScalabilityMode::L3T1_KEY },
+	{ "L3T2", EScalabilityMode::L3T2 },
+	{ "L3T2h", EScalabilityMode::L3T2h },
+	{ "L3T2_KEY", EScalabilityMode::L3T2_KEY },
+	{ "L3T3", EScalabilityMode::L3T3 },
+	{ "L3T3h", EScalabilityMode::L3T3h },
+	{ "L3T3_KEY", EScalabilityMode::L3T3_KEY },
+	{ "S2T1", EScalabilityMode::S2T1 },
+	{ "S2T1h", EScalabilityMode::S2T1h },
+	{ "S2T2", EScalabilityMode::S2T2 },
+	{ "S2T2h", EScalabilityMode::S2T2h },
+	{ "S2T3", EScalabilityMode::S2T3 },
+	{ "S2T3h", EScalabilityMode::S2T3h },
+	{ "S3T1", EScalabilityMode::S3T1 },
+	{ "S3T1h", EScalabilityMode::S3T1h },
+	{ "S3T2", EScalabilityMode::S3T2 },
+	{ "S3T2h", EScalabilityMode::S3T2h },
+	{ "S3T3", EScalabilityMode::S3T3 },
+	{ "S3T3h", EScalabilityMode::S3T3h },
+	{ "None", EScalabilityMode::None }
+};
+
+inline TOptional<EScalabilityMode> ScalabilityModeFromString(const FString& ModeString)
+{
+	if (ScalabilityModeMap.Find(ModeString) != nullptr)
+	{
+		return ScalabilityModeMap[ModeString];
+	}
+
+	return FNullOpt(0);
+}
 
 enum class ERateControlMode : uint8
 {
@@ -26,7 +114,35 @@ enum class EMultipassMode : uint8
 	Full
 };
 
-struct FVideoEncoderConfig : public FAVConfig
+// TODO (william.belcher): Use reasonable defaults set elsewhere
+struct FSpatialLayer
+{
+	uint32 Width = 1920;
+	uint32 Height = 1080;
+	uint32 Framerate = 60;
+	uint8  NumberOfTemporalLayers = 1;
+	int32  MaxBitrate = 20000000;
+	int32  TargetBitrate = 10000000;
+	int32  MinBitrate = 5000000;
+	int32  MaxQP = 0;
+	bool   bActive = false;
+
+	friend bool operator==(const FSpatialLayer& Lhs, const FSpatialLayer& Rhs)
+	{
+		return Lhs.Width == Rhs.Width
+			&& Lhs.Height == Rhs.Height
+			&& Lhs.Framerate == Rhs.Framerate
+			&& Lhs.NumberOfTemporalLayers == Rhs.NumberOfTemporalLayers
+			&& Lhs.MaxBitrate == Rhs.MaxBitrate
+			&& Lhs.TargetBitrate == Rhs.TargetBitrate
+			&& Lhs.MinBitrate == Rhs.MinBitrate
+			&& Lhs.MaxQP == Rhs.MaxQP
+			&& Lhs.bActive == Rhs.bActive;
+	}
+};
+
+// TODO (william.belcher): Use reasonable defaults set elsewhere
+struct FVideoEncoderConfig : public FVideoConfig
 {
 public:
 	uint32 Width = 1920;
@@ -36,73 +152,85 @@ public:
 
 	int32 TargetBitrate = 0;
 	int32 MaxBitrate = 0;
+	int32 MinBitrate = 0;
+	// Advanced bitrate settings. Used for situations such as simulcast / SVC
+	TOptional<int32> Bitrates[Video::MaxSpatialLayers][Video::MaxTemporalStreams];
 
 	int32 MinQP = -1;
 	int32 MaxQP = -1;
 
 	ERateControlMode RateControlMode = ERateControlMode::CBR;
-	uint8 bFillData : 1;
-	
+	uint8			 bFillData : 1;
+
+	EScalabilityMode ScalabilityMode = EScalabilityMode::None;
+
+	uint8				  NumberOfSpatialLayers = 1;
+	uint8				  NumberOfTemporalLayers = 1;
+	TArray<FSpatialLayer> SpatialLayers;
+
+	uint8				  NumberOfSimulcastStreams;
+	TArray<FSpatialLayer> SimulcastStreams;
+
 	// TODO (Remove and derive from latency mode)
-	uint32 KeyframeInterval = 0;
+	uint32		   KeyframeInterval = 0;
 	EMultipassMode MultipassMode = EMultipassMode::Full;
-	
+
 	FVideoEncoderConfig(EAVPreset Preset = EAVPreset::Default)
-		: FAVConfig(Preset)
+		: FVideoConfig(Preset)
 		, bFillData(false)
 	{
 		switch (Preset)
 		{
-		case EAVPreset::UltraLowQuality:
-			TargetBitrate = 500000;
-			MaxBitrate = 1000000;
+			case EAVPreset::UltraLowQuality:
+				TargetBitrate = 500000;
+				MaxBitrate = 1000000;
 
-			MinQP = 10;
-			MaxQP = 20;
+				MinQP = 10;
+				MaxQP = 20;
 
-			RateControlMode = ERateControlMode::CBR;
+				RateControlMode = ERateControlMode::CBR;
 
-			break;
-		case EAVPreset::LowQuality:
-			TargetBitrate = 3000000;
-			MaxBitrate = 4500000;
+				break;
+			case EAVPreset::LowQuality:
+				TargetBitrate = 3000000;
+				MaxBitrate = 4500000;
 
-			MinQP = 20;
-			MaxQP = 30;
+				MinQP = 20;
+				MaxQP = 30;
 
-			RateControlMode = ERateControlMode::CBR;
+				RateControlMode = ERateControlMode::CBR;
 
-			break;
-		case EAVPreset::Default:
-			TargetBitrate = 5000000;
-			MaxBitrate = 12500000;
+				break;
+			case EAVPreset::Default:
+				TargetBitrate = 5000000;
+				MaxBitrate = 12500000;
 
-			MinQP = 25;
-			MaxQP = 40;
+				MinQP = 25;
+				MaxQP = 40;
 
-			RateControlMode = ERateControlMode::CBR;
+				RateControlMode = ERateControlMode::CBR;
 
-			break;
-		case EAVPreset::HighQuality:
-			TargetBitrate = 10000000;
-			MaxBitrate = 20000000;
+				break;
+			case EAVPreset::HighQuality:
+				TargetBitrate = 10000000;
+				MaxBitrate = 20000000;
 
-			MinQP = 35;
-			MaxQP = 50;
+				MinQP = 35;
+				MaxQP = 50;
 
-			RateControlMode = ERateControlMode::VBR;
+				RateControlMode = ERateControlMode::VBR;
 
-			break;
-		case EAVPreset::Lossless:
-			TargetBitrate = 0;
-			MaxBitrate = 0;
+				break;
+			case EAVPreset::Lossless:
+				TargetBitrate = 0;
+				MaxBitrate = 0;
 
-			MinQP = -1;
-			MaxQP = -1;
+				MinQP = -1;
+				MaxQP = -1;
 
-			RateControlMode = ERateControlMode::ConstQP;
+				RateControlMode = ERateControlMode::ConstQP;
 
-			break;
+				break;
 		}
 	}
 };
@@ -189,7 +317,7 @@ public:
 
 		return MinimalConfig;
 	}
-	
+
 	/**
 	 * Set generic configuration values.
 	 *
@@ -222,17 +350,17 @@ public:
 
 	/**
 	 * Flush remaining packets and invalidate the underlying architecture.
-	 * 
+	 *
 	 * @return Result of the operation, @see FAVResult.
 	 */
 	FAVResult FlushPackets()
 	{
 		return SendFrame(nullptr, 0);
 	}
-	
+
 	/**
 	 * Flush remaining packets and invalidate the underlying architecture.
-	 * 
+	 *
 	 * @return Result of the operation, @see FAVResult.
 	 */
 	FAVResult FlushAndReceivePackets(TArray<FVideoPacket>& OutPackets)
@@ -289,7 +417,7 @@ public:
 	 * @return The current minimal configuration.
 	 */
 	virtual FVideoEncoderConfig GetMinimalConfig() = 0;
-	
+
 	/**
 	 * Set generic configuration values.
 	 *
diff --git a/Engine/Plugins/Experimental/AVCodecs/AVCodecsCore/Source/AVCodecsCore/Public/Video/VideoPacket.h b/Engine/Plugins/Experimental/AVCodecs/AVCodecsCore/Source/AVCodecsCore/Public/Video/VideoPacket.h
--- a/Engine/Plugins/Experimental/AVCodecs/AVCodecsCore/Source/AVCodecsCore/Public/Video/VideoPacket.h
+++ b/Engine/Plugins/Experimental/AVCodecs/AVCodecsCore/Source/AVCodecsCore/Public/Video/VideoPacket.h
@@ -4,6 +4,7 @@
 
 #include "AVPacket.h"
 #include "Templates/SharedPointerFwd.h"
+#include "Video/VideoCodecSpecificInfo.h"
 
 /**
  * Video-specific packet container.
@@ -21,6 +22,21 @@ public:
 	 */
 	uint8 bIsKeyframe : 1;
 
+	/**
+	 * Struct holding advanced information on the video packet
+	 */
+	FCodecSpecificInfo CodecSpecificInfo;
+
+	/**
+	 * Spatial index of this packet
+	 */
+	TOptional<int> SpatialIndex;
+
+	/**
+	 * Temporal index of this packet
+	 */
+	TOptional<int> TemporalIndex;
+
 	FVideoPacket() = default;
 	FVideoPacket(TSharedPtr<uint8> const& DataPtr, uint64 DataSize, uint64 Timestamp, uint64 Index, uint32 QP, bool bIsKeyframe)
 		: FAVPacket(DataPtr, DataSize, Timestamp, Index)
diff --git a/Engine/Plugins/Experimental/AVCodecs/AVCodecsCore/Source/AVCodecsCore/Public/Video/VideoResource.h b/Engine/Plugins/Experimental/AVCodecs/AVCodecsCore/Source/AVCodecsCore/Public/Video/VideoResource.h
--- a/Engine/Plugins/Experimental/AVCodecs/AVCodecsCore/Source/AVCodecsCore/Public/Video/VideoResource.h
+++ b/Engine/Plugins/Experimental/AVCodecs/AVCodecsCore/Source/AVCodecsCore/Public/Video/VideoResource.h
@@ -19,6 +19,7 @@ enum class EVideoFormat : uint8
 	G16 = EPixelFormat::PF_G16,
 
 	// TODO No match to engine formats and should be updated if that changes
+	YUV420 = 253U,
 	YUV444 = 254U,
 	YUV444_16 = 255U,
 };
@@ -63,11 +64,11 @@ public:
 	 * Height in pixels.
 	 */
 	uint32 Height;
-    
-    /**
-     * An implementation of the BulkData interface. Used for creating textures from existing data
-     */
-    FResourceBulkDataInterface* BulkData;
+
+	/**
+	 * An implementation of the BulkData interface. Used for creating textures from existing data
+	 */
+	FResourceBulkDataInterface* BulkData;
 
 	/**
 	 * If this is a descriptor stored in a different format than it really is
@@ -80,10 +81,9 @@ public:
 		: Format(Descriptor.Format)
 		, Width(Descriptor.Width)
 		, Height(Descriptor.Height)
-        , BulkData(Descriptor.BulkData)
+		, BulkData(Descriptor.BulkData)
 		, RawDescriptor(Descriptor.RawDescriptor ? new FVideoDescriptor(*Descriptor.RawDescriptor) : nullptr)
 	{
-
 	}
 
 	FVideoDescriptor(EVideoFormat Format, uint32 Width, uint32 Height, FResourceBulkDataInterface* BulkData = nullptr)
@@ -94,44 +94,43 @@ public:
 	FVideoDescriptor(EVideoFormat Format, uint32 Width, uint32 Height, const FVideoDescriptor& RawDescriptor, FResourceBulkDataInterface* BulkData = nullptr)
 		: Format(Format), Width(Width), Height(Height), BulkData(BulkData), RawDescriptor(new FVideoDescriptor(RawDescriptor))
 	{
-
 	}
 
 	~FVideoDescriptor()
-	{ 
-		if(RawDescriptor)
+	{
+		if (RawDescriptor)
 		{
 			delete RawDescriptor;
 		}
 	};
 
-	bool operator==(FVideoDescriptor const &RHS) const
+	bool operator==(FVideoDescriptor const& RHS) const
 	{
 		return Format == RHS.Format && Width == RHS.Width && Height == RHS.Height;
 	}
 
-	bool operator!=(FVideoDescriptor const &RHS) const
+	bool operator!=(FVideoDescriptor const& RHS) const
 	{
 		return !(*this == RHS);
 	}
 
-	// Planar formats are treated as single channel 
+	// Planar formats are treated as single channel
 	uint8 GetNumChannels() const
 	{
 		switch (Format)
 		{
-		case EVideoFormat::BGRA:
-		case EVideoFormat::ABGR10:
-			return 4;
-		case EVideoFormat::YUV444:
-		case EVideoFormat::YUV444_16:
-		case EVideoFormat::NV12:
-		case EVideoFormat::P010:
-		case EVideoFormat::R8:
-			return 1;
-		default:
-			return 0;
-		}		
+			case EVideoFormat::BGRA:
+			case EVideoFormat::ABGR10:
+				return 4;
+			case EVideoFormat::YUV444:
+			case EVideoFormat::YUV444_16:
+			case EVideoFormat::NV12:
+			case EVideoFormat::P010:
+			case EVideoFormat::R8:
+				return 1;
+			default:
+				return 0;
+		}
 	}
 
 	// TODO(Nick): Try and remove this in favour of external initialisation
@@ -140,19 +139,19 @@ public:
 		const uint32 ChannelSize = Width * Height;
 		switch (Format)
 		{
-		case EVideoFormat::BGRA:
-		case EVideoFormat::ABGR10:
-			return ChannelSize * 4;
-		case EVideoFormat::YUV444:
-			return ChannelSize * 3;
-		case EVideoFormat::YUV444_16:
-			return ChannelSize * 6;
-		case EVideoFormat::NV12:
-			return ChannelSize + (ChannelSize >> 1);
-		case EVideoFormat::P010:
-			return ChannelSize * 3;
-		default:
-			return 0;
+			case EVideoFormat::BGRA:
+			case EVideoFormat::ABGR10:
+				return ChannelSize * 4;
+			case EVideoFormat::YUV444:
+				return ChannelSize * 3;
+			case EVideoFormat::YUV444_16:
+				return ChannelSize * 6;
+			case EVideoFormat::NV12:
+				return ChannelSize + (ChannelSize >> 1);
+			case EVideoFormat::P010:
+				return ChannelSize * 3;
+			default:
+				return 0;
 		}
 	}
 };
@@ -172,12 +171,12 @@ public:
 	/**
 	 * @return Get the descriptor of our video data in device memory.
 	 */
-	FORCEINLINE FVideoDescriptor const &GetRawDescriptor() const { return Descriptor.RawDescriptor ? *Descriptor.RawDescriptor : Descriptor; }
+	FORCEINLINE FVideoDescriptor const& GetRawDescriptor() const { return Descriptor.RawDescriptor ? *Descriptor.RawDescriptor : Descriptor; }
 
 	/**
 	 * @return Get the descriptor of our video data.
 	 */
-	FORCEINLINE FVideoDescriptor const &GetDescriptor() const { return Descriptor; }
+	FORCEINLINE FVideoDescriptor const& GetDescriptor() const { return Descriptor; }
 
 	/**
 	 * @return Get the format of our data.
@@ -199,7 +198,7 @@ public:
 	 */
 	FORCEINLINE uint32 GetSize() const { return Descriptor.GetSizeInBytes(); }
 
-	FVideoResource(TSharedRef<FAVDevice> const &Device, FAVLayout const &Layout, FVideoDescriptor const &Descriptor);
+	FVideoResource(TSharedRef<FAVDevice> const& Device, FAVLayout const& Layout, FVideoDescriptor const& Descriptor);
 	virtual ~FVideoResource() override = default;
 };
 
@@ -212,7 +211,7 @@ class TVideoResource : public FVideoResource
 public:
 	FORCEINLINE TSharedPtr<TContext> GetContext() const { return GetDevice()->template GetContext<TContext>(); }
 
-	TVideoResource(TSharedRef<FAVDevice> const &Device, FAVLayout const &Layout, FVideoDescriptor const &Descriptor)
+	TVideoResource(TSharedRef<FAVDevice> const& Device, FAVLayout const& Layout, FVideoDescriptor const& Descriptor)
 		: FVideoResource(Device, Layout, Descriptor)
 	{
 	}
@@ -239,13 +238,13 @@ class TPooledVideoResource : public TResolvableVideoResource<TResource>
 public:
 	TArray<TSharedPtr<TResource>> Pool;
 
-	virtual bool IsResolved(TSharedPtr<FAVDevice> const &Device, FVideoDescriptor const &Descriptor) const override
+	virtual bool IsResolved(TSharedPtr<FAVDevice> const& Device, FVideoDescriptor const& Descriptor) const override
 	{
 		return TResolvableVideoResource<TResource>::IsResolved(Device, Descriptor) && Pool.Contains(this->Get());
 	}
 
 protected:
-	virtual TSharedPtr<TResource> TryResolve(TSharedPtr<FAVDevice> const &Device, FVideoDescriptor const &Descriptor) override
+	virtual TSharedPtr<TResource> TryResolve(TSharedPtr<FAVDevice> const& Device, FVideoDescriptor const& Descriptor) override
 	{
 		for (int i = 0; i < Pool.Num(); ++i)
 		{
diff --git a/Engine/Plugins/Experimental/AVCodecs/LibVpxCodecs/Source/LibVpxCodecs/Public/Video/Decoders/VideoDecoderLibVpxVP9.h b/Engine/Plugins/Experimental/AVCodecs/LibVpxCodecs/Source/LibVpxCodecs/Public/Video/Decoders/VideoDecoderLibVpxVP9.h
new file mode 100644
--- /dev/null
+++ b/Engine/Plugins/Experimental/AVCodecs/LibVpxCodecs/Source/LibVpxCodecs/Public/Video/Decoders/VideoDecoderLibVpxVP9.h
@@ -0,0 +1,59 @@
+// Copyright Epic Games, Inc. All Rights Reserved.
+
+#pragma once
+
+#include "Video/VideoDecoder.h"
+#include "Video/Decoders/Configs/VideoDecoderConfigLibVpx.h"
+#include "Video/Resources/VideoResourceCPU.h"
+
+#include "Containers/Queue.h"
+#include "HAL/Platform.h"
+#include "LibVpx.h"
+
+template <typename TResource>
+class TVideoDecoderLibVpxVP9 : public TVideoDecoder<TResource, FVideoDecoderConfigLibVpx>
+{
+private:
+	uint8 bInitialized : 1;
+
+	uint64 FrameCount = 0;
+
+	bool bIsOpen = false;
+
+	vpx_codec_ctx_t* Decoder = nullptr;
+
+	bool bKeyFrameRequired = true;
+
+	struct FFrame
+	{
+		TArray<uint8> RawData;
+		int			  StrideY = 0;
+		int			  StrideUV = 0;
+		int			  Width = 0;
+		int			  Height = 0;
+		uint64		  FrameNumber = 0;
+	};
+
+	TQueue<FFrame> Frames;
+
+public:
+	TVideoDecoderLibVpxVP9() = default;
+	virtual ~TVideoDecoderLibVpxVP9() override;
+
+	virtual bool	  IsOpen() const override;
+	virtual FAVResult Open(TSharedRef<FAVDevice> const& NewDevice, TSharedRef<FAVInstance> const& NewInstance) override;
+	virtual void	  Close() override;
+
+	bool IsInitialized() const;
+
+	virtual FAVResult ApplyConfig() override;
+
+	virtual FAVResult ReceiveFrame(TResolvableVideoResource<TResource>& InOutResource) override;
+
+	virtual FAVResult SendPacket(FVideoPacket const& Packet) override;
+
+private:
+	FAVResult Destroy();
+};
+
+#include "VideoDecoderLibVpxVP9.hpp"
