Fixing nonunity build errors on mac

diff --git a/Engine/Plugins/Experimental/AVCodecs/AVCodecsCore/Source/AVCodecsCore/Private/Video/Encoders/SVC/VideoBitrateAllocatorSVC.cpp b/Engine/Plugins/Experimental/AVCodecs/AVCodecsCore/Source/AVCodecsCore/Private/Video/Encoders/SVC/VideoBitrateAllocatorSVC.cpp
--- a/Engine/Plugins/Experimental/AVCodecs/AVCodecsCore/Source/AVCodecsCore/Private/Video/Encoders/SVC/VideoBitrateAllocatorSVC.cpp
+++ b/Engine/Plugins/Experimental/AVCodecs/AVCodecsCore/Source/AVCodecsCore/Private/Video/Encoders/SVC/VideoBitrateAllocatorSVC.cpp
@@ -1,7 +1,4 @@
 // Copyright Epic Games, Inc. All Rights Reserved.
-
-#pragma once
-
 #include "Video/Encoders/SVC/VideoBitrateAllocatorSVC.h"
 
 #include "Algo/Accumulate.h"
@@ -17,7 +14,7 @@ struct FActiveSpatialLayers
 	size_t Num = 0;
 };
 
-FActiveSpatialLayers GetActiveSpatialLayers(const FVideoEncoderConfig& Config, size_t NumSpatialLayers)
+FActiveSpatialLayers GetActiveSpatialLayers(const FVideoEncoderConfig &Config, size_t NumSpatialLayers)
 {
 	FActiveSpatialLayers Active;
 	for (Active.First = 0; Active.First < NumSpatialLayers; ++Active.First)
@@ -41,7 +38,7 @@ FActiveSpatialLayers GetActiveSpatialLayers(const FVideoEncoderConfig& Config, s
 	return Active;
 }
 
-TArray<uint32> AdjustAndVerify(const FVideoEncoderConfig& Config, size_t FirstActiveLayer, const TArray<uint32>& SpatialLayerRates)
+TArray<uint32> AdjustAndVerify(const FVideoEncoderConfig &Config, size_t FirstActiveLayer, const TArray<uint32> &SpatialLayerRates)
 {
 	TArray<uint32> AdjustedSpatialLayerRates;
 	// Keep track of rate that couldn't be applied to the previous layer due to
@@ -114,7 +111,7 @@ static TArray<uint32> SplitBitrate(size_t NumLayers, uint32 TotalBitrateBps, flo
 
 // Returns the minimum bitrate needed for `NumActiveLayers` Spatial Layers to
 // become Active using the configuration specified by `Config`.
-uint32 FindLayerTogglingThreshold(const FVideoEncoderConfig& Config, size_t FirstActiveLayer, size_t NumActiveLayers)
+uint32 FindLayerTogglingThreshold(const FVideoEncoderConfig &Config, size_t FirstActiveLayer, size_t NumActiveLayers)
 {
 	if (NumActiveLayers == 1)
 	{
@@ -151,7 +148,7 @@ uint32 FindLayerTogglingThreshold(const FVideoEncoderConfig& Config, size_t Firs
 	return UpperBound;
 }
 
-FVideoBitrateAllocatorSVC::FNumLayers FVideoBitrateAllocatorSVC::GetNumLayers(const FVideoEncoderConfig& Config)
+FVideoBitrateAllocatorSVC::FNumLayers FVideoBitrateAllocatorSVC::GetNumLayers(const FVideoEncoderConfig &Config)
 {
 	FNumLayers Layers;
 	if (Config.ScalabilityMode != EScalabilityMode::None)
@@ -177,11 +174,8 @@ FVideoBitrateAllocatorSVC::FNumLayers FVideoBitrateAllocatorSVC::GetNumLayers(co
 	return Layers;
 }
 
-FVideoBitrateAllocatorSVC::FVideoBitrateAllocatorSVC(const FVideoEncoderConfig& Config)
-	: Config(Config)
-	, NumLayers(GetNumLayers(Config))
-	, CumulativeLayerStartBitrates(GetLayerStartBitrates(Config))
-	, LastActiveLayerCount(0)
+FVideoBitrateAllocatorSVC::FVideoBitrateAllocatorSVC(const FVideoEncoderConfig &Config)
+	: Config(Config), NumLayers(GetNumLayers(Config)), CumulativeLayerStartBitrates(GetLayerStartBitrates(Config)), LastActiveLayerCount(0)
 {
 	check(NumLayers.Spatial > 0);
 	check(NumLayers.Spatial <= 5);
@@ -218,7 +212,7 @@ FVideoBitrateAllocation FVideoBitrateAllocatorSVC::Allocate(FVideoBitrateAllocat
 	}
 
 	const FActiveSpatialLayers ActiveLayers = GetActiveSpatialLayers(Config, NumLayers.Spatial);
-	size_t					   NumSpatialLayers = ActiveLayers.Num;
+	size_t NumSpatialLayers = ActiveLayers.Num;
 
 	NumSpatialLayers = FindNumEnabledLayers(Parameters.TotalBitrateBps);
 	LastActiveLayerCount = NumSpatialLayers;
@@ -303,9 +297,9 @@ size_t FVideoBitrateAllocatorSVC::FindNumEnabledLayers(uint32 TargetRate) const
 	return NumEnabledLayers;
 }
 
-uint32 FVideoBitrateAllocatorSVC::GetMaxBitrate(const FVideoEncoderConfig& Config)
+uint32 FVideoBitrateAllocatorSVC::GetMaxBitrate(const FVideoEncoderConfig &Config)
 {
-	const FNumLayers		   NumLayers = GetNumLayers(Config);
+	const FNumLayers NumLayers = GetNumLayers(Config);
 	const FActiveSpatialLayers ActiveLayers = GetActiveSpatialLayers(Config, NumLayers.Spatial);
 
 	uint32 MaxBitrate = 0.f;
@@ -322,7 +316,7 @@ uint32 FVideoBitrateAllocatorSVC::GetMaxBitrate(const FVideoEncoderConfig& Confi
 	return MaxBitrate;
 }
 
-uint32 FVideoBitrateAllocatorSVC::GetPaddingBitrate(const FVideoEncoderConfig& Config)
+uint32 FVideoBitrateAllocatorSVC::GetPaddingBitrate(const FVideoEncoderConfig &Config)
 {
 	auto StartBitrate = GetLayerStartBitrates(Config);
 	if (StartBitrate.IsEmpty())
@@ -333,12 +327,12 @@ uint32 FVideoBitrateAllocatorSVC::GetPaddingBitrate(const FVideoEncoderConfig& C
 	return StartBitrate.Last();
 }
 
-TArray<uint32> FVideoBitrateAllocatorSVC::GetLayerStartBitrates(const FVideoEncoderConfig& Config)
+TArray<uint32> FVideoBitrateAllocatorSVC::GetLayerStartBitrates(const FVideoEncoderConfig &Config)
 {
-	TArray<uint32>			   StartBitrates;
-	const FNumLayers		   NumLayers = GetNumLayers(Config);
+	TArray<uint32> StartBitrates;
+	const FNumLayers NumLayers = GetNumLayers(Config);
 	const FActiveSpatialLayers ActiveLayers = GetActiveSpatialLayers(Config, NumLayers.Spatial);
-	uint32					   LastRate = 0.f;
+	uint32 LastRate = 0.f;
 	for (size_t i = 1; i <= ActiveLayers.Num; ++i)
 	{
 		uint32 LayerTogglingRate = FindLayerTogglingThreshold(Config, ActiveLayers.First, i);
diff --git a/Engine/Plugins/Experimental/AVCodecs/LibVpxCodecs/Source/LibVpxCodecs/Public/Video/Decoders/VideoDecoderLibVpxVP9.h b/Engine/Plugins/Experimental/AVCodecs/LibVpxCodecs/Source/LibVpxCodecs/Public/Video/Decoders/VideoDecoderLibVpxVP9.h
--- a/Engine/Plugins/Experimental/AVCodecs/LibVpxCodecs/Source/LibVpxCodecs/Public/Video/Decoders/VideoDecoderLibVpxVP9.h
+++ b/Engine/Plugins/Experimental/AVCodecs/LibVpxCodecs/Source/LibVpxCodecs/Public/Video/Decoders/VideoDecoderLibVpxVP9.h
@@ -4,6 +4,7 @@
 
 #include "Video/VideoDecoder.h"
 #include "Video/Decoders/Configs/VideoDecoderConfigLibVpx.h"
+#include "Video/Decoders/Configs/VideoDecoderConfigVP9.h"
 #include "Video/Resources/VideoResourceCPU.h"
 
 #include "Containers/Queue.h"
@@ -20,18 +21,18 @@ private:
 
 	bool bIsOpen = false;
 
-	vpx_codec_ctx_t* Decoder = nullptr;
+	vpx_codec_ctx_t *Decoder = nullptr;
 
 	bool bKeyFrameRequired = true;
 
 	struct FFrame
 	{
 		TArray<uint8> RawData;
-		int			  StrideY = 0;
-		int			  StrideUV = 0;
-		int			  Width = 0;
-		int			  Height = 0;
-		uint64		  FrameNumber = 0;
+		int StrideY = 0;
+		int StrideUV = 0;
+		int Width = 0;
+		int Height = 0;
+		uint64 FrameNumber = 0;
 	};
 
 	TQueue<FFrame> Frames;
@@ -40,17 +41,17 @@ public:
 	TVideoDecoderLibVpxVP9() = default;
 	virtual ~TVideoDecoderLibVpxVP9() override;
 
-	virtual bool	  IsOpen() const override;
-	virtual FAVResult Open(TSharedRef<FAVDevice> const& NewDevice, TSharedRef<FAVInstance> const& NewInstance) override;
-	virtual void	  Close() override;
+	virtual bool IsOpen() const override;
+	virtual FAVResult Open(TSharedRef<FAVDevice> const &NewDevice, TSharedRef<FAVInstance> const &NewInstance) override;
+	virtual void Close() override;
 
 	bool IsInitialized() const;
 
 	virtual FAVResult ApplyConfig() override;
 
-	virtual FAVResult ReceiveFrame(TResolvableVideoResource<TResource>& InOutResource) override;
+	virtual FAVResult ReceiveFrame(TResolvableVideoResource<TResource> &InOutResource) override;
 
-	virtual FAVResult SendPacket(FVideoPacket const& Packet) override;
+	virtual FAVResult SendPacket(FVideoPacket const &Packet) override;
 
 private:
 	FAVResult Destroy();
