Fix: TArray's of FSpatialLayers could cause hangs during shutdown

diff --git a/Engine/Plugins/Experimental/AVCodecs/AVCodecsCore/Source/AVCodecsCore/Public/Video/VideoEncoder.h b/Engine/Plugins/Experimental/AVCodecs/AVCodecsCore/Source/AVCodecsCore/Public/Video/VideoEncoder.h
--- a/Engine/Plugins/Experimental/AVCodecs/AVCodecsCore/Source/AVCodecsCore/Public/Video/VideoEncoder.h
+++ b/Engine/Plugins/Experimental/AVCodecs/AVCodecsCore/Source/AVCodecsCore/Public/Video/VideoEncoder.h
@@ -169,12 +169,12 @@ public:
 
 	EScalabilityMode ScalabilityMode = EScalabilityMode::None;
 
-	uint8				  NumberOfSpatialLayers = 1;
-	uint8				  NumberOfTemporalLayers = 1;
-	TArray<FSpatialLayer> SpatialLayers;
+	uint8		  NumberOfSpatialLayers = 1;
+	uint8		  NumberOfTemporalLayers = 1;
+	FSpatialLayer SpatialLayers[Video::MaxSpatialLayers];
 
-	uint8				  NumberOfSimulcastStreams;
-	TArray<FSpatialLayer> SimulcastStreams;
+	uint8		  NumberOfSimulcastStreams;
+	FSpatialLayer SimulcastStreams[Video::MaxSimulcastStreams];
 
 	// TODO (Remove and derive from latency mode)
 	uint32		   KeyframeInterval = 0;
@@ -186,56 +186,56 @@ public:
 	{
 		switch (Preset)
 		{
-		case EAVPreset::UltraLowQuality:
-			TargetBitrate = 500000;
-			MaxBitrate = 1000000;
+			case EAVPreset::UltraLowQuality:
+				TargetBitrate = 500000;
+				MaxBitrate = 1000000;
 
-			MinQuality = 0;
-			MaxQuality = 33;
+				MinQuality = 0;
+				MaxQuality = 33;
 
-			RateControlMode = ERateControlMode::CBR;
+				RateControlMode = ERateControlMode::CBR;
 
-			break;
-		case EAVPreset::LowQuality:
-			TargetBitrate = 3000000;
-			MaxBitrate = 4500000;
+				break;
+			case EAVPreset::LowQuality:
+				TargetBitrate = 3000000;
+				MaxBitrate = 4500000;
 
-			MinQuality = 0;
-			MaxQuality = 50;
+				MinQuality = 0;
+				MaxQuality = 50;
 
-			RateControlMode = ERateControlMode::CBR;
+				RateControlMode = ERateControlMode::CBR;
 
-			break;
-		case EAVPreset::Default:
-			TargetBitrate = 5000000;
-			MaxBitrate = 12500000;
+				break;
+			case EAVPreset::Default:
+				TargetBitrate = 5000000;
+				MaxBitrate = 12500000;
 
-			MinQuality = 25;
-			MaxQuality = 75;
+				MinQuality = 25;
+				MaxQuality = 75;
 
-			RateControlMode = ERateControlMode::CBR;
+				RateControlMode = ERateControlMode::CBR;
 
-			break;
-		case EAVPreset::HighQuality:
-			TargetBitrate = 10000000;
-			MaxBitrate = 20000000;
+				break;
+			case EAVPreset::HighQuality:
+				TargetBitrate = 10000000;
+				MaxBitrate = 20000000;
 
-			MinQuality = 50;
-			MaxQuality = 100;
+				MinQuality = 50;
+				MaxQuality = 100;
 
-			RateControlMode = ERateControlMode::VBR;
+				RateControlMode = ERateControlMode::VBR;
 
-			break;
-		case EAVPreset::Lossless:
-			TargetBitrate = 0;
-			MaxBitrate = 0;
+				break;
+			case EAVPreset::Lossless:
+				TargetBitrate = 0;
+				MaxBitrate = 0;
 
-			MinQuality = -1;
-			MaxQuality = -1;
+				MinQuality = -1;
+				MaxQuality = -1;
 
-			RateControlMode = ERateControlMode::ConstQP;
+				RateControlMode = ERateControlMode::ConstQP;
 
-			break;
+				break;
 		}
 	}
 };
diff --git a/Engine/Plugins/Experimental/AVCodecs/LibVpxCodecs/Source/LibVpxCodecs/Private/Video/Encoders/Configs/VideoEncoderConfigLibVpx.cpp b/Engine/Plugins/Experimental/AVCodecs/LibVpxCodecs/Source/LibVpxCodecs/Private/Video/Encoders/Configs/VideoEncoderConfigLibVpx.cpp
--- a/Engine/Plugins/Experimental/AVCodecs/LibVpxCodecs/Source/LibVpxCodecs/Private/Video/Encoders/Configs/VideoEncoderConfigLibVpx.cpp
+++ b/Engine/Plugins/Experimental/AVCodecs/LibVpxCodecs/Source/LibVpxCodecs/Private/Video/Encoders/Configs/VideoEncoderConfigLibVpx.cpp
@@ -27,10 +27,9 @@ DLLEXPORT FAVResult FAVExtension::TransformConfig(FVideoEncoderConfigLibVpx& Out
 	OutConfig.ScalabilityMode = InConfig.ScalabilityMode;
 	OutConfig.NumberOfSpatialLayers = InConfig.NumberOfSpatialLayers;
 	OutConfig.NumberOfTemporalLayers = InConfig.NumberOfTemporalLayers;
-	OutConfig.SpatialLayers = InConfig.SpatialLayers;
+	FMemory::Memcpy(OutConfig.SpatialLayers, InConfig.SpatialLayers, sizeof(FSpatialLayer) * Video::MaxSpatialLayers);
 	OutConfig.NumberOfSimulcastStreams = InConfig.NumberOfSimulcastStreams;
-	OutConfig.SimulcastStreams = InConfig.SimulcastStreams;
-
+	FMemory::Memcpy(OutConfig.SimulcastStreams, InConfig.SimulcastStreams, sizeof(FSpatialLayer) * Video::MaxSimulcastStreams);
 	for (size_t si = 0; si < Video::MaxSpatialLayers; si++)
 	{
 		for (size_t ti = 0; ti < Video::MaxTemporalStreams; ti++)
@@ -61,9 +60,9 @@ DLLEXPORT FAVResult FAVExtension::TransformConfig(FVideoEncoderConfig& OutConfig
 	OutConfig.ScalabilityMode = InConfig.ScalabilityMode;
 	OutConfig.NumberOfSpatialLayers = InConfig.NumberOfSpatialLayers;
 	OutConfig.NumberOfTemporalLayers = InConfig.NumberOfTemporalLayers;
-	OutConfig.SpatialLayers = InConfig.SpatialLayers;
+	FMemory::Memcpy(OutConfig.SpatialLayers, InConfig.SpatialLayers, sizeof(FSpatialLayer) * Video::MaxSpatialLayers);
 	OutConfig.NumberOfSimulcastStreams = InConfig.NumberOfSimulcastStreams;
-	OutConfig.SimulcastStreams = InConfig.SimulcastStreams;
+	FMemory::Memcpy(OutConfig.SimulcastStreams, InConfig.SimulcastStreams, sizeof(FSpatialLayer) * Video::MaxSimulcastStreams);
 
 	for (size_t si = 0; si < Video::MaxSpatialLayers; si++)
 	{
diff --git a/Engine/Plugins/Experimental/AVCodecs/LibVpxCodecs/Source/LibVpxCodecs/Public/Video/Encoders/Configs/VideoEncoderConfigLibVpx.h b/Engine/Plugins/Experimental/AVCodecs/LibVpxCodecs/Source/LibVpxCodecs/Public/Video/Encoders/Configs/VideoEncoderConfigLibVpx.h
--- a/Engine/Plugins/Experimental/AVCodecs/LibVpxCodecs/Source/LibVpxCodecs/Public/Video/Encoders/Configs/VideoEncoderConfigLibVpx.h
+++ b/Engine/Plugins/Experimental/AVCodecs/LibVpxCodecs/Source/LibVpxCodecs/Public/Video/Encoders/Configs/VideoEncoderConfigLibVpx.h
@@ -25,19 +25,19 @@ public:
 	int32			 MinQuality;
 	int32			 MaxQuality;
 
-	int32					 NumberOfCores = 0;
-	bool					 bDenoisingOn = false;
-	bool					 bAdaptiveQpMode = false;
-	bool					 bAutomaticResizeOn = false;
-	bool					 bFlexibleMode = false;
+	int32				  NumberOfCores = 0;
+	bool				  bDenoisingOn = false;
+	bool				  bAdaptiveQpMode = false;
+	bool				  bAutomaticResizeOn = false;
+	bool				  bFlexibleMode = false;
 	EInterLayerPrediction InterLayerPrediction = EInterLayerPrediction::Off;
 
-	uint8				  NumberOfSpatialLayers = 1;
-	uint8				  NumberOfTemporalLayers = 1;
-	TArray<FSpatialLayer> SpatialLayers;
+	uint8		  NumberOfSpatialLayers = 1;
+	uint8		  NumberOfTemporalLayers = 1;
+	FSpatialLayer SpatialLayers[Video::MaxSpatialLayers];
 
-	uint8				  NumberOfSimulcastStreams;
-	TArray<FSpatialLayer> SimulcastStreams;
+	uint8		  NumberOfSimulcastStreams;
+	FSpatialLayer SimulcastStreams[Video::MaxSimulcastStreams];
 
 	EScalabilityMode ScalabilityMode = EScalabilityMode::None;
 
@@ -78,11 +78,11 @@ public:
 			&& this->InterLayerPrediction == Other.InterLayerPrediction
 			&& this->NumberOfSpatialLayers == Other.NumberOfSpatialLayers
 			&& this->NumberOfTemporalLayers == Other.NumberOfTemporalLayers
-			&& this->SpatialLayers == Other.SpatialLayers
 			&& this->NumberOfSimulcastStreams == Other.NumberOfSimulcastStreams
-			&& this->SimulcastStreams == Other.SimulcastStreams
 			&& this->ScalabilityMode == Other.ScalabilityMode
-			&& this->SameBitrates(Other.Bitrates);
+			&& this->SameBitrates(Other.Bitrates)
+			&& this->SameSpatialLayers(Other.SpatialLayers)
+			&& this->SameSimulcastStreams(Other.SimulcastStreams);
 	}
 
 	bool operator!=(FVideoEncoderConfigLibVpx const& Other) const
@@ -105,6 +105,32 @@ public:
 
 		return true;
 	}
+
+	bool SameSpatialLayers(const FSpatialLayer OtherSpatialLayers[Video::MaxSpatialLayers]) const
+	{
+		for (size_t si = 0; si < Video::MaxSpatialLayers; si++)
+		{
+			if (SpatialLayers[si] != OtherSpatialLayers[si])
+			{
+				return false;
+			}
+		}
+
+		return true;
+	}
+
+	bool SameSimulcastStreams(const FSpatialLayer OtherSimulcastStreams[Video::MaxSimulcastStreams]) const
+	{
+		for (size_t si = 0; si < Video::MaxSimulcastStreams; si++)
+		{
+			if (SimulcastStreams[si] != OtherSimulcastStreams[si])
+			{
+				return false;
+			}
+		}
+
+		return true;
+	}
 };
 
 template <>
diff --git a/Engine/Plugins/Experimental/AVCodecs/LibVpxCodecs/Source/LibVpxCodecs/Public/Video/Encoders/VideoEncoderLibVpxVP9.hpp b/Engine/Plugins/Experimental/AVCodecs/LibVpxCodecs/Source/LibVpxCodecs/Public/Video/Encoders/VideoEncoderLibVpxVP9.hpp
--- a/Engine/Plugins/Experimental/AVCodecs/LibVpxCodecs/Source/LibVpxCodecs/Public/Video/Encoders/VideoEncoderLibVpxVP9.hpp
+++ b/Engine/Plugins/Experimental/AVCodecs/LibVpxCodecs/Source/LibVpxCodecs/Public/Video/Encoders/VideoEncoderLibVpxVP9.hpp
@@ -185,6 +185,7 @@ FAVResult TVideoEncoderLibVpxVP9<TResource>::ApplyConfig()
 				else
 				{
 					NumSpatialLayers = PendingConfig.NumberOfSpatialLayers;
+					check(NumSpatialLayers > 0);
 					NumTemporalLayers = PendingConfig.NumberOfTemporalLayers;
 					if (NumTemporalLayers == 0)
 					{
@@ -489,6 +490,9 @@ FAVResult TVideoEncoderLibVpxVP9<TResource>::SendFrame(TSharedPtr<FVideoResource
 			bVpxConfigChanged = false;
 		}
 
+		check(Resource->GetWidth() == RawImage->d_w);
+		check(Resource->GetHeight() == RawImage->d_h);
+
 		InputImage = MakeUnique<FInputImage>(InTimestamp);
 
 		switch (Profile)
@@ -720,11 +724,11 @@ FAVResult TVideoEncoderLibVpxVP9<TResource>::InitAndSetControlSettings(FVideoEnc
 			SvcParams->scaling_factor_den[i] = ScaleFactor;
 
 			check(Config.SpatialLayers[i].Framerate > 0);
-			check(Config.SpatialLayers[i].Framerate < Config.Framerate);
+			check(Config.SpatialLayers[i].Framerate <= Config.Framerate);
 			if (i > 0)
 			{
 				// Frame rate of high spatial Layer is supposed to be equal or higher than frame rate of low spatial Layer.
-				check(Config.SpatialLayers[i].Framerate > Config.SpatialLayers[i - 1].Framerate);
+				check(Config.SpatialLayers[i].Framerate >= Config.SpatialLayers[i - 1].Framerate);
 			}
 		}
 	}
