[AVCodecs] Fix: VTCodecs hardware decoding. NOTE: There is a known memory leak with this codepath but will be addressed in a later CL

diff --git a/Engine/Plugins/Experimental/AVCodecs/AVCodecsCore/Source/AVCodecsCoreRHI/Private/Video/Resources/VideoResourceRHI.cpp b/Engine/Plugins/Experimental/AVCodecs/AVCodecsCore/Source/AVCodecsCoreRHI/Private/Video/Resources/VideoResourceRHI.cpp
--- a/Engine/Plugins/Experimental/AVCodecs/AVCodecsCore/Source/AVCodecsCoreRHI/Private/Video/Resources/VideoResourceRHI.cpp
+++ b/Engine/Plugins/Experimental/AVCodecs/AVCodecsCore/Source/AVCodecsCoreRHI/Private/Video/Resources/VideoResourceRHI.cpp
@@ -142,8 +142,13 @@ TSharedPtr<FVideoResourceRHI> FVideoResourceRHI::Create(TSharedPtr<FAVDevice> co
 		FRHITextureCreateDesc TextureDesc = FRHITextureCreateDesc::Create2D(TEXT("AVCodecs Resource"), Descriptor.Width, Descriptor.Height, static_cast<EPixelFormat>(Descriptor.Format));
 
 		TextureDesc.SetClearValue(FClearValueBinding::None);
+#if AVCODECS_USE_METAL
+		TextureDesc.SetFlags(ETextureCreateFlags::CPUReadback);
+		TextureDesc.SetInitialState(ERHIAccess::CPURead);
+#else
     	TextureDesc.SetFlags(ETextureCreateFlags::RenderTargetable);
 		TextureDesc.SetInitialState(ERHIAccess::Present);
+#endif
 
 		TextureDesc.SetNumMips(1);
 
diff --git a/Engine/Plugins/Experimental/AVCodecs/VTCodecs/Source/VTCodecs/Public/Video/Decoders/Configs/VideoDecoderConfigVT.h b/Engine/Plugins/Experimental/AVCodecs/VTCodecs/Source/VTCodecs/Public/Video/Decoders/Configs/VideoDecoderConfigVT.h
--- a/Engine/Plugins/Experimental/AVCodecs/VTCodecs/Source/VTCodecs/Public/Video/Decoders/Configs/VideoDecoderConfigVT.h
+++ b/Engine/Plugins/Experimental/AVCodecs/VTCodecs/Source/VTCodecs/Public/Video/Decoders/Configs/VideoDecoderConfigVT.h
@@ -17,7 +17,7 @@ public:
 	
 public:
 	CMVideoCodecType Codec;
-    CMVideoFormatDescriptionRef VideoFormat;
+    CMVideoFormatDescriptionRef VideoFormat = nullptr;
 
 	FVideoDecoderConfigVT()
 		: FAVConfig()
