[AVCodecs] Fix: Crash and green frames when using VulkanRHI with PixelStreamingPlayer

diff --git a/Engine/Plugins/Experimental/AVCodecs/AVCodecsCore/Source/AVCodecsCoreRHI/Private/Video/Resources/VideoResourceRHI.cpp b/Engine/Plugins/Experimental/AVCodecs/AVCodecsCore/Source/AVCodecsCoreRHI/Private/Video/Resources/VideoResourceRHI.cpp
--- a/Engine/Plugins/Experimental/AVCodecs/AVCodecsCore/Source/AVCodecsCoreRHI/Private/Video/Resources/VideoResourceRHI.cpp
+++ b/Engine/Plugins/Experimental/AVCodecs/AVCodecsCore/Source/AVCodecsCoreRHI/Private/Video/Resources/VideoResourceRHI.cpp
@@ -161,11 +161,6 @@ TSharedPtr<FVideoResourceRHI> FVideoResourceRHI::Create(TSharedPtr<FAVDevice> co
 			TextureDesc.AddFlags(ETextureCreateFlags::Shared);
 		}
 
-		if(bIsSRGB)
-		{
-			TextureDesc.AddFlags(ETextureCreateFlags::SRGB);
-		}
-
 		TextureDesc.AddFlags(AdditionalFlags);
         
         if (Descriptor.BulkData)
@@ -201,13 +196,37 @@ TSharedPtr<FVideoResourceRHI> FVideoResourceRHI::Create(TSharedPtr<FAVDevice> co
         case EVideoFormat::BGRA:
 			TextureDesc.Format = EPixelFormat::PF_B8G8R8A8;
 			Descriptor.RawDescriptor = new FVideoDescriptor(EVideoFormat::BGRA, TextureDesc.Extent.X, TextureDesc.Extent.Y);
+			// Only BGRA can be directly displayed and could need sRGB adjustment, setting this for all formats triggers a check		
+			// in VulkanRHI
+			if (bIsSRGB)
+			{
+				TextureDesc.AddFlags(ETextureCreateFlags::SRGB);
+			}
 			break;
 		default:
 			break;
 		}
 
-//TODO-TE THIS IS THE REAL DEAL?
-		return MakeShareable(new FVideoResourceRHI(Device.ToSharedRef(), { GDynamicRHI->RHICreateTexture(FRHICommandListExecutor::GetImmediateCommandList(), TextureDesc), nullptr, 0 }, Descriptor));
+		// We have to add a fence here as VulkanRHI adds a fill to the command buffer that might not be dispatched until
+		// after a external decoder API has already copied to the texture. If we had access to Signal Semaphores they could
+		// instead be shared to the other APIs instead.
+		FGPUFenceRHIRef Fence = GDynamicRHI->RHICreateGPUFence("CreateVideoResourceRHI");
+		
+		// There seems to be an intermitant bug with VulkanRHI that results in crahes if we dont enque the allocation 
+		// of this texture
+		FTextureRHIRef Texture;
+		ENQUEUE_RENDER_COMMAND(CreateTexture)
+			([&Texture, TextureDesc, Fence](FRHICommandListImmediate& RHICmdList) {
+				Texture = GDynamicRHI->RHICreateTexture(RHICmdList, TextureDesc);
+				RHICmdList.WriteGPUFence(Fence);
+				});
+
+		while (!Fence->Poll()) {
+			constexpr float SleepTimeSeconds = 50 * 1E-6;
+			FPlatformProcess::SleepNoStats(SleepTimeSeconds);
+		}
+
+		return MakeShareable(new FVideoResourceRHI(Device.ToSharedRef(), { Texture, nullptr, 0 }, Descriptor));
 	}
 
 	return nullptr;
diff --git a/Engine/Plugins/Media/PixelStreaming2/Source/PixelStreaming2/Private/Blueprints/PixelStreaming2MediaTexture.cpp b/Engine/Plugins/Media/PixelStreaming2/Source/PixelStreaming2/Private/Blueprints/PixelStreaming2MediaTexture.cpp
--- a/Engine/Plugins/Media/PixelStreaming2/Source/PixelStreaming2/Private/Blueprints/PixelStreaming2MediaTexture.cpp
+++ b/Engine/Plugins/Media/PixelStreaming2/Source/PixelStreaming2/Private/Blueprints/PixelStreaming2MediaTexture.cpp
@@ -111,7 +111,7 @@ FTextureResource* UPixelStreaming2MediaTexture::CreateResource()
 
 void UPixelStreaming2MediaTexture::ConsumeFrame(FTextureRHIRef Frame)
 {
-	AsyncTask(ENamedThreads::ActualRenderingThread, [this, Frame]() {
+	AsyncTask(ENamedThreads::GetRenderThread(), [this, Frame]() {
 		FScopeLock Lock(&RenderSyncContext);
 
 		FRHICommandListImmediate& RHICmdList = FRHICommandListImmediate::Get();
diff --git a/Engine/Plugins/Media/PixelStreaming2/Source/PixelStreaming2/Private/EpicRtcVideoDecoder.cpp b/Engine/Plugins/Media/PixelStreaming2/Source/PixelStreaming2/Private/EpicRtcVideoDecoder.cpp
--- a/Engine/Plugins/Media/PixelStreaming2/Source/PixelStreaming2/Private/EpicRtcVideoDecoder.cpp
+++ b/Engine/Plugins/Media/PixelStreaming2/Source/PixelStreaming2/Private/EpicRtcVideoDecoder.cpp
@@ -116,7 +116,8 @@ namespace UE::PixelStreaming2
 			FAVResult DecodeResult;
 			if constexpr (std::is_same_v<TVideoResource, FVideoResourceRHI>)
 			{
-				FResolvableVideoResourceRHI DecoderResource;
+				// HACK edge case if streaming framerate is faster than render framerate this will cause an issue and should be switched to FetchOrCreate
+				FResolvableVideoResourceRHI& DecoderResource = VideoResourcesRHI[VideoResourceIndex++ % 2];
 				DecodeResult = Decoder->ReceiveFrame(DecoderResource);
 
 				if (DecodeResult.IsSuccess())
@@ -126,7 +127,8 @@ namespace UE::PixelStreaming2
 			}
 			else if constexpr (std::is_same_v<TVideoResource, FVideoResourceCPU>)
 			{
-				FResolvableVideoResourceCPU DecoderResource;
+				// HACK edge case if streaming framerate is faster than render framerate this will cause an issue and should be switched to FetchOrCreate
+				FResolvableVideoResourceCPU& DecoderResource = VideoResourcesCPU[VideoResourceIndex++ % 2];
 				DecodeResult = Decoder->ReceiveFrame(DecoderResource);
 
 				if (DecodeResult.IsSuccess())
diff --git a/Engine/Plugins/Media/PixelStreaming2/Source/PixelStreaming2/Private/EpicRtcVideoDecoder.h b/Engine/Plugins/Media/PixelStreaming2/Source/PixelStreaming2/Private/EpicRtcVideoDecoder.h
--- a/Engine/Plugins/Media/PixelStreaming2/Source/PixelStreaming2/Private/EpicRtcVideoDecoder.h
+++ b/Engine/Plugins/Media/PixelStreaming2/Source/PixelStreaming2/Private/EpicRtcVideoDecoder.h
@@ -13,6 +13,8 @@
 
 #include "epic_rtc/core/video/video_decoder.h"
 
+#include <atomic>
+
 namespace UE::PixelStreaming2
 {
 	template <std::derived_from<FVideoResource> TVideoResource>
@@ -36,7 +38,10 @@ namespace UE::PixelStreaming2
 		TRefCountPtr<EpicRtcVideoDecoderCallbackInterface> VideoDecoderCallback;
 		TRefCountPtr<EpicRtcVideoCodecInfoInterface>	   CodecInfo;
 		uint16_t										   FrameCount;
-
+		// TODO rather than just flip flopping we should have a system to FetchOrCreate a Resource on an as needed basis
+		TStaticArray<FResolvableVideoResourceRHI, 2>	   VideoResourcesRHI;
+		TStaticArray<FResolvableVideoResourceCPU, 2>	   VideoResourcesCPU;
+		std::atomic<uint32>								   VideoResourceIndex;
 	private:
 		bool LateInitDecoder();
 
diff --git a/Engine/Plugins/Media/PixelStreaming2/Source/PixelStreaming2/Private/VideoSink.cpp b/Engine/Plugins/Media/PixelStreaming2/Source/PixelStreaming2/Private/VideoSink.cpp
--- a/Engine/Plugins/Media/PixelStreaming2/Source/PixelStreaming2/Private/VideoSink.cpp
+++ b/Engine/Plugins/Media/PixelStreaming2/Source/PixelStreaming2/Private/VideoSink.cpp
@@ -108,7 +108,7 @@ namespace UE::PixelStreaming2
 					FrameBuffer->GetWidth(), FrameBuffer->GetHeight());
 			}
 
-			AsyncTask(ENamedThreads::ActualRenderingThread, [Width, Height, this]() {
+			AsyncTask(ENamedThreads::GetRenderThread(), [Width, Height, this]() {
 				FScopeLock Lock(&RenderSyncContext);
 
 				const FIntPoint			  FrameSize = FIntPoint(Width, Height);
