Replace some usages of FORCEINLINE with inline in AVCodec modules.

For context about this change, see this doc: https://docs.google.com/document/d/1LjpT7NpzTecfXdFNyv5KM6sCu67827BRljNi1iaWVSE/edit?tab=t.0#heading=h.fvtz8gjukz7x

Includes modules: AVCodecsCore, AVCodecsCoreRHI, NVCodecs

Reviewers were automatically determined. Please comment if anyone else should be added.

#rnx

diff --git a/Engine/Plugins/Experimental/AVCodecs/AVCodecsCore/Source/AVCodecsCore/Public/AVPacket.h b/Engine/Plugins/Experimental/AVCodecs/AVCodecsCore/Source/AVCodecsCore/Public/AVPacket.h
--- a/Engine/Plugins/Experimental/AVCodecs/AVCodecsCore/Source/AVCodecsCore/Public/AVPacket.h
+++ b/Engine/Plugins/Experimental/AVCodecs/AVCodecsCore/Source/AVCodecsCore/Public/AVPacket.h
@@ -38,9 +38,9 @@ struct FAVPacket
 	 *
 	 * @return Raw data as an array view.
 	 */
-	FORCEINLINE TArrayView64<uint8> GetData() const { return TArrayView64<uint8>(DataPtr.Get(), DataSize); }
+	inline TArrayView64<uint8> GetData() const { return TArrayView64<uint8>(DataPtr.Get(), DataSize); }
 
-	FORCEINLINE void WriteToFile(TCHAR* FileName) const {
+	inline void WriteToFile(TCHAR* FileName) const {
 		FString SaveName = FString::Printf(TEXT("%s%s"), *FPaths::ProjectSavedDir(), FileName);
 
 		UE_LOG(LogTemp, Display, TEXT("Saving Packet to %s"), *SaveName);
diff --git a/Engine/Plugins/Experimental/AVCodecs/AVCodecsCore/Source/AVCodecsCore/Public/AVResource.h b/Engine/Plugins/Experimental/AVCodecs/AVCodecsCore/Source/AVCodecsCore/Public/AVResource.h
--- a/Engine/Plugins/Experimental/AVCodecs/AVCodecsCore/Source/AVCodecsCore/Public/AVResource.h
+++ b/Engine/Plugins/Experimental/AVCodecs/AVCodecsCore/Source/AVCodecsCore/Public/AVResource.h
@@ -38,27 +38,27 @@ public:
 	/**
 	 * @return Get a reference to our device parent. Is always valid.
 	 */
-	FORCEINLINE TSharedRef<FAVDevice> const& GetDevice() const { return Device; }
+	inline TSharedRef<FAVDevice> const& GetDevice() const { return Device; }
 
 	/**
 	 * @return Get the layout of our data in device memory.
 	 */
-	FORCEINLINE FAVLayout const& GetLayout() const { return Layout; }
+	inline FAVLayout const& GetLayout() const { return Layout; }
 
 	/**
 	 * @return Get the offset of our data in device memory.
 	 */
-	FORCEINLINE uint64 GetOffset() const { return Layout.Offset; }
+	inline uint64 GetOffset() const { return Layout.Offset; }
 
 	/**
 	 * @return Get the size of our data in device memory.
 	 */
-	FORCEINLINE uint64 GetMemorySize() const { return Layout.Size; }
+	inline uint64 GetMemorySize() const { return Layout.Size; }
 	
 	/**
 	 * @return Get the stride of our data in device memory.
 	 */
-	FORCEINLINE uint64 GetStride() const { return Layout.Stride; }
+	inline uint64 GetStride() const { return Layout.Stride; }
 
 	FAVResource(TSharedRef<FAVDevice> const& Device, FAVLayout const& Layout);
 	virtual ~FAVResource() = default;
diff --git a/Engine/Plugins/Experimental/AVCodecs/AVCodecsCore/Source/AVCodecsCore/Public/AVResult.h b/Engine/Plugins/Experimental/AVCodecs/AVCodecsCore/Source/AVCodecsCore/Public/AVResult.h
--- a/Engine/Plugins/Experimental/AVCodecs/AVCodecsCore/Source/AVCodecsCore/Public/AVResult.h
+++ b/Engine/Plugins/Experimental/AVCodecs/AVCodecsCore/Source/AVCodecsCore/Public/AVResult.h
@@ -95,7 +95,7 @@ public:
 	 * @param Vendor Optional name of vendor where this error occured.
 	 * @param VendorValue Optional result code passed to us by the vendor.
 	 */
-	FORCEINLINE static void Log(EAVResult Value, FString const& Message = "", FString const& Vendor = "", int32 VendorValue = EMPTY_VENDOR_VALUE)
+	inline static void Log(EAVResult Value, FString const& Message = "", FString const& Vendor = "", int32 VendorValue = EMPTY_VENDOR_VALUE)
 	{
 		FAVResult(Value, Message, Vendor, VendorValue).Log();
 	}
@@ -161,7 +161,7 @@ public:
 		}
 	}
 
-	FORCEINLINE FAVResult& operator=(FAVResult const& From)
+	inline FAVResult& operator=(FAVResult const& From)
 	{
 		Value = From.Value;
 		Message = From.Message;
@@ -174,17 +174,17 @@ public:
 		return *this;
 	}
 
-	FORCEINLINE bool operator==(EAVResult OtherValue) const
+	inline bool operator==(EAVResult OtherValue) const
 	{
 		return Value == OtherValue;
 	}
 
-	FORCEINLINE bool operator!=(EAVResult OtherValue) const
+	inline bool operator!=(EAVResult OtherValue) const
 	{
 		return !(*this == OtherValue);
 	}
 
-	FORCEINLINE operator bool() const
+	inline operator bool() const
 	{
 		return IsSuccess();
 	}
@@ -192,7 +192,7 @@ public:
 	/**
 	 * @return True if AV result code is in the range of Fatal.
 	 */
-	FORCEINLINE bool IsFatal() const
+	inline bool IsFatal() const
 	{
 		return Value >= EAVResult::Fatal && Value < EAVResult::Error;
 	}
@@ -200,7 +200,7 @@ public:
 	/**
 	 * @return True if AV result code is NOT in the range of Fatal.
 	 */
-	FORCEINLINE bool IsNotFatal() const
+	inline bool IsNotFatal() const
 	{
 		return !IsFatal();
 	}
@@ -208,7 +208,7 @@ public:
 	/**
 	 * @return True if AV result code is in the range of Error.
 	 */
-	FORCEINLINE bool IsError() const
+	inline bool IsError() const
 	{
 		return Value >= EAVResult::Error && Value < EAVResult::Warning;
 	}
@@ -216,7 +216,7 @@ public:
 	/**
 	 * @return True if AV result code is NOT in the range of Error.
 	 */
-	FORCEINLINE bool IsNotError() const
+	inline bool IsNotError() const
 	{
 		return !IsError();
 	}
@@ -224,7 +224,7 @@ public:
 	/**
 	 * @return True if AV result code is in the range of Warning.
 	 */
-	FORCEINLINE bool IsWarning() const
+	inline bool IsWarning() const
 	{
 		return Value >= EAVResult::Warning && Value < EAVResult::Pending;
 	}
@@ -232,7 +232,7 @@ public:
 	/**
 	 * @return True if AV result code is NOT in the range of Warning.
 	 */
-	FORCEINLINE bool IsNotWarning() const
+	inline bool IsNotWarning() const
 	{
 		return !IsWarning();
 	}
@@ -240,7 +240,7 @@ public:
 	/**
 	 * @return True if AV result code is in the range of Pending.
 	 */
-	FORCEINLINE bool IsPending() const
+	inline bool IsPending() const
 	{
 		return Value >= EAVResult::Pending;
 	}
@@ -248,7 +248,7 @@ public:
 	/**
 	 * @return True if AV result code is NOT in the range of Pending.
 	 */
-	FORCEINLINE bool IsNotPending() const
+	inline bool IsNotPending() const
 	{
 		return !IsPending();
 	}
@@ -256,7 +256,7 @@ public:
 	/**
 	 * @return True if AV result code is in the range of Success.
 	 */
-	FORCEINLINE bool IsSuccess() const
+	inline bool IsSuccess() const
 	{
 		return Value >= EAVResult::Success;
 	}
@@ -264,7 +264,7 @@ public:
 	/**
 	 * @return True if AV result code is NOT in the range of Success.
 	 */
-	FORCEINLINE bool IsNotSuccess() const
+	inline bool IsNotSuccess() const
 	{
 		return !IsSuccess();
 	}
@@ -331,12 +331,12 @@ public:
 	{
 	}
 
-	FORCEINLINE operator TReturnValue&()
+	inline operator TReturnValue&()
 	{
 		return ReturnValue;
 	}
 
-	FORCEINLINE operator TReturnValue const&() const
+	inline operator TReturnValue const&() const
 	{
 		return ReturnValue;
 	}
diff --git a/Engine/Plugins/Experimental/AVCodecs/AVCodecsCore/Source/AVCodecsCore/Public/Audio/AudioResource.h b/Engine/Plugins/Experimental/AVCodecs/AVCodecsCore/Source/AVCodecsCore/Public/Audio/AudioResource.h
--- a/Engine/Plugins/Experimental/AVCodecs/AVCodecsCore/Source/AVCodecsCore/Public/Audio/AudioResource.h
+++ b/Engine/Plugins/Experimental/AVCodecs/AVCodecsCore/Source/AVCodecsCore/Public/Audio/AudioResource.h
@@ -81,17 +81,17 @@ public:
 	/**
 	 * @return Get the descriptor of our audio data in device memory.
 	 */
-	FORCEINLINE FAudioDescriptor const& GetDescriptor() const { return Descriptor; }
+	inline FAudioDescriptor const& GetDescriptor() const { return Descriptor; }
 	
 	/**
 	 * @return Get the number of samples in our data in device memory.
 	 */
-	FORCEINLINE float GetNumSamples() const { return Descriptor.NumSamples; }
+	inline float GetNumSamples() const { return Descriptor.NumSamples; }
 
 	/**
 	 * @return Get the duration of our data in device memory.
 	 */
-	FORCEINLINE float GetSampleDuration() const { return Descriptor.SampleDuration; }
+	inline float GetSampleDuration() const { return Descriptor.SampleDuration; }
 
 	FAudioResource(TSharedRef<FAVDevice> const& Device, FAVLayout const& Layout, FAudioDescriptor const& Descriptor);
 	virtual ~FAudioResource() override = default;
@@ -104,7 +104,7 @@ template <typename TContext>
 class TAudioResource : public FAudioResource
 {
 public:
-	FORCEINLINE TSharedPtr<TContext> GetContext() const { return GetDevice()->template GetContext<TContext>(); }
+	inline TSharedPtr<TContext> GetContext() const { return GetDevice()->template GetContext<TContext>(); }
 	
 	TAudioResource(TSharedRef<FAVDevice> const& Device, FAVLayout const& Layout, FAudioDescriptor const& Descriptor)
 		: FAudioResource(Device, Layout, Descriptor)
diff --git a/Engine/Plugins/Experimental/AVCodecs/AVCodecsCore/Source/AVCodecsCore/Public/Audio/Resources/AudioResourceCPU.h b/Engine/Plugins/Experimental/AVCodecs/AVCodecsCore/Source/AVCodecsCore/Public/Audio/Resources/AudioResourceCPU.h
--- a/Engine/Plugins/Experimental/AVCodecs/AVCodecsCore/Source/AVCodecsCore/Public/Audio/Resources/AudioResourceCPU.h
+++ b/Engine/Plugins/Experimental/AVCodecs/AVCodecsCore/Source/AVCodecsCore/Public/Audio/Resources/AudioResourceCPU.h
@@ -21,7 +21,7 @@ private:
 	TSharedPtr<float> Raw;
 
 public:
-	FORCEINLINE TSharedPtr<float> const& GetRaw() const { return Raw; }
+	inline TSharedPtr<float> const& GetRaw() const { return Raw; }
 
 	FAudioResourceCPU(TSharedRef<FAVDevice> const& Device, TSharedPtr<float> const& Raw, FAVLayout const& Layout, FAudioDescriptor const& Descriptor);
 	virtual ~FAudioResourceCPU() override = default;
diff --git a/Engine/Plugins/Experimental/AVCodecs/AVCodecsCore/Source/AVCodecsCore/Public/Video/Resources/D3D/VideoResourceD3D.h b/Engine/Plugins/Experimental/AVCodecs/AVCodecsCore/Source/AVCodecsCore/Public/Video/Resources/D3D/VideoResourceD3D.h
--- a/Engine/Plugins/Experimental/AVCodecs/AVCodecsCore/Source/AVCodecsCore/Public/Video/Resources/D3D/VideoResourceD3D.h
+++ b/Engine/Plugins/Experimental/AVCodecs/AVCodecsCore/Source/AVCodecsCore/Public/Video/Resources/D3D/VideoResourceD3D.h
@@ -60,8 +60,8 @@ private:
 public:
 	static FVideoDescriptor GetDescriptorFrom(TSharedRef<FAVDevice> const& Device, TRefCountPtr<ID3D11Texture2D> const& Raw);
 
-	FORCEINLINE TRefCountPtr<ID3D11Texture2D> const& GetRaw() const { return Raw; }
-	FORCEINLINE HANDLE const& GetSharedHandle() const { return SharedHandle; }
+	inline TRefCountPtr<ID3D11Texture2D> const& GetRaw() const { return Raw; }
+	inline HANDLE const& GetSharedHandle() const { return SharedHandle; }
 
 	FVideoResourceD3D11(TSharedRef<FAVDevice> const& Device, TRefCountPtr<ID3D11Texture2D> const& Raw, FAVLayout const& Layout);
 
@@ -88,15 +88,15 @@ private:
 public:
 	static FVideoDescriptor GetDescriptorFrom(TSharedRef<FAVDevice> const& Device, FRawD3D12 const& Raw);
 
-	FORCEINLINE TRefCountPtr<ID3D12Resource> const& GetResource() const { return Raw.D3DResource; }
-	FORCEINLINE HANDLE const& GetResourceSharedHandle() const { return Raw.D3DResourceShared; }
+	inline TRefCountPtr<ID3D12Resource> const& GetResource() const { return Raw.D3DResource; }
+	inline HANDLE const& GetResourceSharedHandle() const { return Raw.D3DResourceShared; }
 	
-	FORCEINLINE TRefCountPtr<ID3D12Heap> const& GetHeap() const { return Raw.D3DHeap; }
-	FORCEINLINE HANDLE const& GetHeapSharedHandle() const { return Raw.D3DHeapShared; }
+	inline TRefCountPtr<ID3D12Heap> const& GetHeap() const { return Raw.D3DHeap; }
+	inline HANDLE const& GetHeapSharedHandle() const { return Raw.D3DHeapShared; }
 
-	FORCEINLINE TRefCountPtr<ID3D12Fence> const& GetFence() const { return Raw.D3DFence; };
-	FORCEINLINE HANDLE const& GetFenceSharedHandle() const { return Raw.D3DFenceShared; }
-	FORCEINLINE uint64 const& GetFenceValue() const { return Raw.FenceValue; }
+	inline TRefCountPtr<ID3D12Fence> const& GetFence() const { return Raw.D3DFence; };
+	inline HANDLE const& GetFenceSharedHandle() const { return Raw.D3DFenceShared; }
+	inline uint64 const& GetFenceValue() const { return Raw.FenceValue; }
 
 	uint64 GetSizeInBytes();
 
diff --git a/Engine/Plugins/Experimental/AVCodecs/AVCodecsCore/Source/AVCodecsCore/Public/Video/Resources/Metal/VideoResourceMetal.h b/Engine/Plugins/Experimental/AVCodecs/AVCodecsCore/Source/AVCodecsCore/Public/Video/Resources/Metal/VideoResourceMetal.h
--- a/Engine/Plugins/Experimental/AVCodecs/AVCodecsCore/Source/AVCodecsCore/Public/Video/Resources/Metal/VideoResourceMetal.h
+++ b/Engine/Plugins/Experimental/AVCodecs/AVCodecsCore/Source/AVCodecsCore/Public/Video/Resources/Metal/VideoResourceMetal.h
@@ -36,7 +36,7 @@ private:
 public:
 	static FVideoDescriptor GetDescriptorFrom(TSharedRef<FAVDevice> const& Device, MTL::Texture* Raw);
 
-	FORCEINLINE MTL::Texture* GetRaw() const { return Raw; }
+	inline MTL::Texture* GetRaw() const { return Raw; }
 
 	FVideoResourceMetal(TSharedRef<FAVDevice> const& Device, MTL::Texture* Raw, FAVLayout const& Layout);
     virtual ~FVideoResourceMetal() override;
diff --git a/Engine/Plugins/Experimental/AVCodecs/AVCodecsCore/Source/AVCodecsCore/Public/Video/Resources/VideoResourceCPU.h b/Engine/Plugins/Experimental/AVCodecs/AVCodecsCore/Source/AVCodecsCore/Public/Video/Resources/VideoResourceCPU.h
--- a/Engine/Plugins/Experimental/AVCodecs/AVCodecsCore/Source/AVCodecsCore/Public/Video/Resources/VideoResourceCPU.h
+++ b/Engine/Plugins/Experimental/AVCodecs/AVCodecsCore/Source/AVCodecsCore/Public/Video/Resources/VideoResourceCPU.h
@@ -21,8 +21,8 @@ private:
 	TSharedPtr<uint8> Raw;
 
 public:
-	FORCEINLINE TSharedPtr<uint8> const& GetRaw() const { return Raw; }
-	FORCEINLINE void					 SetRaw(TSharedPtr<uint8> RawData) { Raw = RawData; }
+	inline TSharedPtr<uint8> const& GetRaw() const { return Raw; }
+	inline void					 SetRaw(TSharedPtr<uint8> RawData) { Raw = RawData; }
 
 	FVideoResourceCPU(TSharedRef<FAVDevice> const& Device, TSharedPtr<uint8> const& Raw, FAVLayout const& Layout, FVideoDescriptor const& Descriptor);
 	virtual ~FVideoResourceCPU() override = default;
diff --git a/Engine/Plugins/Experimental/AVCodecs/AVCodecsCore/Source/AVCodecsCore/Public/Video/Resources/Vulkan/VideoResourceVulkan.h b/Engine/Plugins/Experimental/AVCodecs/AVCodecsCore/Source/AVCodecsCore/Public/Video/Resources/Vulkan/VideoResourceVulkan.h
--- a/Engine/Plugins/Experimental/AVCodecs/AVCodecsCore/Source/AVCodecsCore/Public/Video/Resources/Vulkan/VideoResourceVulkan.h
+++ b/Engine/Plugins/Experimental/AVCodecs/AVCodecsCore/Source/AVCodecsCore/Public/Video/Resources/Vulkan/VideoResourceVulkan.h
@@ -47,9 +47,9 @@ private:
 	FVulkanSharedHandle FenceSharedHandle = FVulkanSharedHandle();
 
 public:
-	FORCEINLINE VkDeviceMemory GetRaw() const { return DeviceMemory; }
-	FORCEINLINE FVulkanSharedHandle const& GetSharedHandle() const { return DeviceMemorySharedHandle; }
-	FORCEINLINE FVulkanSharedHandle const& GetFenceSharedHandle() const { return FenceSharedHandle; }
+	inline VkDeviceMemory GetRaw() const { return DeviceMemory; }
+	inline FVulkanSharedHandle const& GetSharedHandle() const { return DeviceMemorySharedHandle; }
+	inline FVulkanSharedHandle const& GetFenceSharedHandle() const { return FenceSharedHandle; }
 
 	FVideoResourceVulkan(TSharedRef<FAVDevice> const& Device, VkDeviceMemory Raw, FAVLayout const& Layout, FVideoDescriptor const& Descriptor);
 	virtual ~FVideoResourceVulkan() override = default;
diff --git a/Engine/Plugins/Experimental/AVCodecs/AVCodecsCore/Source/AVCodecsCore/Public/Video/VideoResource.h b/Engine/Plugins/Experimental/AVCodecs/AVCodecsCore/Source/AVCodecsCore/Public/Video/VideoResource.h
--- a/Engine/Plugins/Experimental/AVCodecs/AVCodecsCore/Source/AVCodecsCore/Public/Video/VideoResource.h
+++ b/Engine/Plugins/Experimental/AVCodecs/AVCodecsCore/Source/AVCodecsCore/Public/Video/VideoResource.h
@@ -178,42 +178,42 @@ public:
 	/**
 	 * @return Get the descriptor of our video data in device memory.
 	 */
-	FORCEINLINE FVideoDescriptor const& GetRawDescriptor() const { return Descriptor.RawDescriptor ? *Descriptor.RawDescriptor : Descriptor; }
+	inline FVideoDescriptor const& GetRawDescriptor() const { return Descriptor.RawDescriptor ? *Descriptor.RawDescriptor : Descriptor; }
 
 	/**
 	 * @return Get the descriptor of our video data.
 	 */
-	FORCEINLINE FVideoDescriptor const& GetDescriptor() const { return Descriptor; }
+	inline FVideoDescriptor const& GetDescriptor() const { return Descriptor; }
 
 	/**
 	 * @return Get the format of our data.
 	 */
-	FORCEINLINE EVideoFormat GetFormat() const { return Descriptor.Format; }
+	inline EVideoFormat GetFormat() const { return Descriptor.Format; }
 
 	/**
 	 * @return Get the width of our data.
 	 */
-	FORCEINLINE uint32 GetWidth() const { return Descriptor.Width; }
+	inline uint32 GetWidth() const { return Descriptor.Width; }
 
 	/**
 	 * @return Get the height of our data.
 	 */
-	FORCEINLINE uint32 GetHeight() const { return Descriptor.Height; }
+	inline uint32 GetHeight() const { return Descriptor.Height; }
 
 	/**
 	 * @return Get the Size in bytes of our Resource.
 	 */
-	FORCEINLINE uint32 GetSize() const { return Descriptor.GetSizeInBytes(); }
+	inline uint32 GetSize() const { return Descriptor.GetSizeInBytes(); }
 
 	/**
 	 * 
 	 */
-	FORCEINLINE void SetUsing(bool bUsing) { bIsInUse = bUsing; }
+	inline void SetUsing(bool bUsing) { bIsInUse = bUsing; }
 
 	/**
 	 * 
 	 */
-	FORCEINLINE bool IsInUse() { return bIsInUse; }
+	inline bool IsInUse() { return bIsInUse; }
 
 	FVideoResource(TSharedRef<FAVDevice> const& Device, FAVLayout const& Layout, FVideoDescriptor const& Descriptor);
 	virtual ~FVideoResource() override = default;
@@ -226,7 +226,7 @@ template <typename TContext>
 class TVideoResource : public FVideoResource
 {
 public:
-	FORCEINLINE TSharedPtr<TContext> GetContext() const { return GetDevice()->template GetContext<TContext>(); }
+	inline TSharedPtr<TContext> GetContext() const { return GetDevice()->template GetContext<TContext>(); }
 
 	TVideoResource(TSharedRef<FAVDevice> const& Device, FAVLayout const& Layout, FVideoDescriptor const& Descriptor)
 		: FVideoResource(Device, Layout, Descriptor)
diff --git a/Engine/Plugins/Experimental/AVCodecs/AVCodecsCore/Source/AVCodecsCoreRHI/Public/Video/Resources/VideoResourceRHI.h b/Engine/Plugins/Experimental/AVCodecs/AVCodecsCore/Source/AVCodecsCoreRHI/Public/Video/Resources/VideoResourceRHI.h
--- a/Engine/Plugins/Experimental/AVCodecs/AVCodecsCore/Source/AVCodecsCoreRHI/Public/Video/Resources/VideoResourceRHI.h
+++ b/Engine/Plugins/Experimental/AVCodecs/AVCodecsCore/Source/AVCodecsCoreRHI/Public/Video/Resources/VideoResourceRHI.h
@@ -42,7 +42,7 @@ public:
 	static TSharedPtr<FVideoResourceRHI> Create(TSharedPtr<FAVDevice> const& Device, FVideoDescriptor const& Descriptor, ETextureCreateFlags AdditionalFlags = ETextureCreateFlags::None, bool bIsSRGB = true);
 
 
-	FORCEINLINE FRawData const& GetRaw() const { return Raw; }
+	inline FRawData const& GetRaw() const { return Raw; }
 
 	FVideoResourceRHI(TSharedRef<FAVDevice> const& Device, FRawData const& Raw, FVideoDescriptor const& OverrideDescriptor);
 	FVideoResourceRHI(TSharedRef<FAVDevice> const& Device, FRawData const& Raw);
diff --git a/Engine/Plugins/Experimental/AVCodecs/NVCodecs/Source/NVCodecs/Public/Video/Resources/VideoResourceCUDA.h b/Engine/Plugins/Experimental/AVCodecs/NVCodecs/Source/NVCodecs/Public/Video/Resources/VideoResourceCUDA.h
--- a/Engine/Plugins/Experimental/AVCodecs/NVCodecs/Source/NVCodecs/Public/Video/Resources/VideoResourceCUDA.h
+++ b/Engine/Plugins/Experimental/AVCodecs/NVCodecs/Source/NVCodecs/Public/Video/Resources/VideoResourceCUDA.h
@@ -41,7 +41,7 @@ private:
 	CUarray Raw = nullptr;
 
 public:
-	FORCEINLINE CUarray GetRaw() const { return Raw; }
+	inline CUarray GetRaw() const { return Raw; }
 
 	FVideoResourceCUDA(TSharedRef<FAVDevice> const& Device, CUarray Raw, FAVLayout const& Layout, FVideoDescriptor const& Descriptor);
 	FVideoResourceCUDA(TSharedRef<FAVDevice> const& Device, CUDA_EXTERNAL_MEMORY_HANDLE_DESC const& ExternalResourceDesc, CUDA_EXTERNAL_SEMAPHORE_HANDLE_DESC const& ExternalFenceDesc, FAVLayout const& Layout, FVideoDescriptor const& Descriptor, uint64 FenceValue = 0);
@@ -59,7 +59,7 @@ public:
 	FAVResult CopyToAsync(CUdeviceptr Source, uint32 MapPitch);
 
 private:
-	FORCEINLINE CUDA_MEMCPY2D CreateCopyFromParams(CUdeviceptr Source, uint32 SourcePitch)
+	inline CUDA_MEMCPY2D CreateCopyFromParams(CUdeviceptr Source, uint32 SourcePitch)
 	{
 		CUDA_MEMCPY2D CopyParams = {};
 		CopyParams.srcMemoryType = CU_MEMORYTYPE_DEVICE;
@@ -73,7 +73,7 @@ private:
 		return CopyParams;
 	}
 
-	FORCEINLINE CUDA_MEMCPY2D CreateCopyToParams(CUdeviceptr Target, uint32 TargetPitch)
+	inline CUDA_MEMCPY2D CreateCopyToParams(CUdeviceptr Target, uint32 TargetPitch)
 	{
 		CUDA_MEMCPY2D CopyParams = {};
 		CopyParams.srcMemoryType = CU_MEMORYTYPE_ARRAY;
