Fix V1043: A global object variable is declared in the header warnings

#rnx

diff --git a/Engine/Plugins/Animation/LiveLinkHub/Source/LiveLinkHub/Private/Clients/SLiveLinkHubClientsView.h b/Engine/Plugins/Animation/LiveLinkHub/Source/LiveLinkHub/Private/Clients/SLiveLinkHubClientsView.h
--- a/Engine/Plugins/Animation/LiveLinkHub/Source/LiveLinkHub/Private/Clients/SLiveLinkHubClientsView.h
+++ b/Engine/Plugins/Animation/LiveLinkHub/Source/LiveLinkHub/Private/Clients/SLiveLinkHubClientsView.h
@@ -42,12 +42,12 @@ DECLARE_DELEGATE_OneParam(FOnDiscoveredClientPicked, FLiveLinkHubClientId/*Clien
 DECLARE_DELEGATE_OneParam(FOnRemoveClientFromSession, FLiveLinkHubClientId/*ClientIdentifier*/);
 
 
-static const FName NameColumnId = "Name";
-static const FName StatusColumnId = "Status";
-static const FName EnabledIconColumnId = "EnabledIcon";
-static const FName DisabledIconName = "LiveLinkHub.AutoConnect.Disabled";
-static const FName GlobalIconName = "LiveLinkHub.AutoConnect.Global";
-static const FName LocalOnlyIconName = "LiveLinkHub.AutoConnect.Local";
+static const FName NameColumnId = "Name"; //-V1043
+static const FName StatusColumnId = "Status"; //-V1043
+static const FName EnabledIconColumnId = "EnabledIcon"; //-V1043
+static const FName DisabledIconName = "LiveLinkHub.AutoConnect.Disabled"; //-V1043
+static const FName GlobalIconName = "LiveLinkHub.AutoConnect.Global"; //-V1043
+static const FName LocalOnlyIconName = "LiveLinkHub.AutoConnect.Local"; //-V1043
 
 namespace LiveLinkHubClientsViewUtils
 {
diff --git a/Engine/Plugins/Experimental/AVCodecs/AVCodecsCore/Source/AVCodecsCore/Private/Video/VideoEncoder.cpp b/Engine/Plugins/Experimental/AVCodecs/AVCodecsCore/Source/AVCodecsCore/Private/Video/VideoEncoder.cpp
--- a/Engine/Plugins/Experimental/AVCodecs/AVCodecsCore/Source/AVCodecsCore/Private/Video/VideoEncoder.cpp
+++ b/Engine/Plugins/Experimental/AVCodecs/AVCodecsCore/Source/AVCodecsCore/Private/Video/VideoEncoder.cpp
@@ -3,3 +3,41 @@
 #include "Video/VideoEncoder.h"
 
 REGISTER_TYPEID(FVideoEncoder);
+
+TMap<FString, EScalabilityMode> const ScalabilityModeMap = {
+	{ "L1T1", EScalabilityMode::L1T1 },
+	{ "L1T2", EScalabilityMode::L1T2 },
+	{ "L1T3", EScalabilityMode::L1T3 },
+	{ "L2T1", EScalabilityMode::L2T1 },
+	{ "L2T1h", EScalabilityMode::L2T1h },
+	{ "L2T1_KEY", EScalabilityMode::L2T1_KEY },
+	{ "L2T2", EScalabilityMode::L2T2 },
+	{ "L2T2h", EScalabilityMode::L2T2h },
+	{ "L2T2_KEY", EScalabilityMode::L2T2_KEY },
+	{ "L2T2_KEY_SHIFT", EScalabilityMode::L2T2_KEY_SHIFT },
+	{ "L2T3", EScalabilityMode::L2T3 },
+	{ "L2T3h", EScalabilityMode::L2T3h },
+	{ "L2T3_KEY", EScalabilityMode::L2T3_KEY },
+	{ "L3T1", EScalabilityMode::L3T1 },
+	{ "L3T1h", EScalabilityMode::L3T1h },
+	{ "L3T1_KEY", EScalabilityMode::L3T1_KEY },
+	{ "L3T2", EScalabilityMode::L3T2 },
+	{ "L3T2h", EScalabilityMode::L3T2h },
+	{ "L3T2_KEY", EScalabilityMode::L3T2_KEY },
+	{ "L3T3", EScalabilityMode::L3T3 },
+	{ "L3T3h", EScalabilityMode::L3T3h },
+	{ "L3T3_KEY", EScalabilityMode::L3T3_KEY },
+	{ "S2T1", EScalabilityMode::S2T1 },
+	{ "S2T1h", EScalabilityMode::S2T1h },
+	{ "S2T2", EScalabilityMode::S2T2 },
+	{ "S2T2h", EScalabilityMode::S2T2h },
+	{ "S2T3", EScalabilityMode::S2T3 },
+	{ "S2T3h", EScalabilityMode::S2T3h },
+	{ "S3T1", EScalabilityMode::S3T1 },
+	{ "S3T1h", EScalabilityMode::S3T1h },
+	{ "S3T2", EScalabilityMode::S3T2 },
+	{ "S3T2h", EScalabilityMode::S3T2h },
+	{ "S3T3", EScalabilityMode::S3T3 },
+	{ "S3T3h", EScalabilityMode::S3T3h },
+	{ "None", EScalabilityMode::None }
+};
\ No newline at end of file
diff --git a/Engine/Plugins/Experimental/AVCodecs/AVCodecsCore/Source/AVCodecsCore/Public/Video/VideoEncoder.h b/Engine/Plugins/Experimental/AVCodecs/AVCodecsCore/Source/AVCodecsCore/Public/Video/VideoEncoder.h
--- a/Engine/Plugins/Experimental/AVCodecs/AVCodecsCore/Source/AVCodecsCore/Public/Video/VideoEncoder.h
+++ b/Engine/Plugins/Experimental/AVCodecs/AVCodecsCore/Source/AVCodecsCore/Public/Video/VideoEncoder.h
@@ -53,43 +53,7 @@ enum class EScalabilityMode : uint8
 	None
 };
 
-TMap<FString, EScalabilityMode> const ScalabilityModeMap = {
-	{ "L1T1", EScalabilityMode::L1T1 },
-	{ "L1T2", EScalabilityMode::L1T2 },
-	{ "L1T3", EScalabilityMode::L1T3 },
-	{ "L2T1", EScalabilityMode::L2T1 },
-	{ "L2T1h", EScalabilityMode::L2T1h },
-	{ "L2T1_KEY", EScalabilityMode::L2T1_KEY },
-	{ "L2T2", EScalabilityMode::L2T2 },
-	{ "L2T2h", EScalabilityMode::L2T2h },
-	{ "L2T2_KEY", EScalabilityMode::L2T2_KEY },
-	{ "L2T2_KEY_SHIFT", EScalabilityMode::L2T2_KEY_SHIFT },
-	{ "L2T3", EScalabilityMode::L2T3 },
-	{ "L2T3h", EScalabilityMode::L2T3h },
-	{ "L2T3_KEY", EScalabilityMode::L2T3_KEY },
-	{ "L3T1", EScalabilityMode::L3T1 },
-	{ "L3T1h", EScalabilityMode::L3T1h },
-	{ "L3T1_KEY", EScalabilityMode::L3T1_KEY },
-	{ "L3T2", EScalabilityMode::L3T2 },
-	{ "L3T2h", EScalabilityMode::L3T2h },
-	{ "L3T2_KEY", EScalabilityMode::L3T2_KEY },
-	{ "L3T3", EScalabilityMode::L3T3 },
-	{ "L3T3h", EScalabilityMode::L3T3h },
-	{ "L3T3_KEY", EScalabilityMode::L3T3_KEY },
-	{ "S2T1", EScalabilityMode::S2T1 },
-	{ "S2T1h", EScalabilityMode::S2T1h },
-	{ "S2T2", EScalabilityMode::S2T2 },
-	{ "S2T2h", EScalabilityMode::S2T2h },
-	{ "S2T3", EScalabilityMode::S2T3 },
-	{ "S2T3h", EScalabilityMode::S2T3h },
-	{ "S3T1", EScalabilityMode::S3T1 },
-	{ "S3T1h", EScalabilityMode::S3T1h },
-	{ "S3T2", EScalabilityMode::S3T2 },
-	{ "S3T2h", EScalabilityMode::S3T2h },
-	{ "S3T3", EScalabilityMode::S3T3 },
-	{ "S3T3h", EScalabilityMode::S3T3h },
-	{ "None", EScalabilityMode::None }
-};
+AVCODECSCORE_API extern TMap<FString, EScalabilityMode> const ScalabilityModeMap;
 
 inline TOptional<EScalabilityMode> ScalabilityModeFromString(const FString& ModeString)
 {
diff --git a/Engine/Plugins/Interchange/Runtime/Source/Import/Private/Gltf/InterchangeGLTFMaterial.cpp b/Engine/Plugins/Interchange/Runtime/Source/Import/Private/Gltf/InterchangeGLTFMaterial.cpp
--- a/Engine/Plugins/Interchange/Runtime/Source/Import/Private/Gltf/InterchangeGLTFMaterial.cpp
+++ b/Engine/Plugins/Interchange/Runtime/Source/Import/Private/Gltf/InterchangeGLTFMaterial.cpp
@@ -20,6 +20,8 @@
 #include "Misc/Paths.h"
 #include "Misc/PackageName.h"
 
+const FString InterchangeGltfMaterialAttributeIdentifier = TEXT("Gltf_MI_AttributeIdentifier_");
+
 namespace UE::Interchange::GLTFMaterials
 {
 
diff --git a/Engine/Plugins/Interchange/Runtime/Source/Import/Public/Gltf/InterchangeGLTFMaterial.h b/Engine/Plugins/Interchange/Runtime/Source/Import/Public/Gltf/InterchangeGLTFMaterial.h
--- a/Engine/Plugins/Interchange/Runtime/Source/Import/Public/Gltf/InterchangeGLTFMaterial.h
+++ b/Engine/Plugins/Interchange/Runtime/Source/Import/Public/Gltf/InterchangeGLTFMaterial.h
@@ -10,7 +10,7 @@ namespace GLTF {
 class UInterchangeShaderGraphNode;
 class UInterchangeBaseNodeContainer;
 
-const FString InterchangeGltfMaterialAttributeIdentifier = TEXT("Gltf_MI_AttributeIdentifier_");
+INTERCHANGEIMPORT_API extern const FString InterchangeGltfMaterialAttributeIdentifier;
 
 //OffsetScale for MaterialInstances
 //Offset_X Offset_Y Scale_X Scale_Y for Materials
diff --git a/Engine/Plugins/Media/PixelStreaming/Source/PixelStreaming/Private/PixelStreamingModule.cpp b/Engine/Plugins/Media/PixelStreaming/Source/PixelStreaming/Private/PixelStreamingModule.cpp
--- a/Engine/Plugins/Media/PixelStreaming/Source/PixelStreaming/Private/PixelStreamingModule.cpp
+++ b/Engine/Plugins/Media/PixelStreaming/Source/PixelStreaming/Private/PixelStreamingModule.cpp
@@ -65,6 +65,9 @@ DEFINE_LOG_CATEGORY(LogPixelStreaming);
 
 IPixelStreamingModule* UE::PixelStreaming::FPixelStreamingModule::PixelStreamingModule = nullptr;
 
+const FPixelStreamingPlayerId INVALID_PLAYER_ID = ToPlayerId(FString(TEXT("Invalid Player Id")));
+const FPixelStreamingPlayerId SFU_PLAYER_ID = FString(TEXT("1"));
+
 namespace UE::PixelStreaming
 {
 	/**
diff --git a/Engine/Plugins/Media/PixelStreaming/Source/PixelStreaming/Public/PixelStreamingPlayerId.h b/Engine/Plugins/Media/PixelStreaming/Source/PixelStreaming/Public/PixelStreamingPlayerId.h
--- a/Engine/Plugins/Media/PixelStreaming/Source/PixelStreaming/Public/PixelStreamingPlayerId.h
+++ b/Engine/Plugins/Media/PixelStreaming/Source/PixelStreaming/Public/PixelStreamingPlayerId.h
@@ -21,5 +21,5 @@ PIXELSTREAMING_API inline int32 PlayerIdToInt(FPixelStreamingPlayerId PlayerId)
 	return FCString::Atoi(*PlayerId);
 }
 
-static const FPixelStreamingPlayerId INVALID_PLAYER_ID = ToPlayerId(FString(TEXT("Invalid Player Id")));
-static const FPixelStreamingPlayerId SFU_PLAYER_ID = FString(TEXT("1"));
\ No newline at end of file
+PIXELSTREAMING_API extern const FPixelStreamingPlayerId INVALID_PLAYER_ID;
+PIXELSTREAMING_API extern const FPixelStreamingPlayerId SFU_PLAYER_ID;
\ No newline at end of file
diff --git a/Engine/Source/Developer/CQTest/Private/CQTestSettings.cpp b/Engine/Source/Developer/CQTest/Private/CQTestSettings.cpp
--- a/Engine/Source/Developer/CQTest/Private/CQTestSettings.cpp
+++ b/Engine/Source/Developer/CQTest/Private/CQTestSettings.cpp
@@ -1,10 +1,14 @@
 // Copyright Epic Games, Inc. All Rights Reserved.
 
 #include "CQTestSettings.h"
+#include "CQTest.h"
 #include "Tests/AutomationCommon.h"
 
 #include UE_INLINE_GENERATED_CPP_BY_NAME(CQTestSettings)
 
+const FString GenerateTestDirectory;
+const FString DefaultTags;
+
 namespace CQTestConsoleVariables 
 {
 static TAutoConsoleVariable<float> CVarCommandTimeout(
diff --git a/Engine/Source/Developer/CQTest/Public/CQTest.h b/Engine/Source/Developer/CQTest/Public/CQTest.h
--- a/Engine/Source/Developer/CQTest/Public/CQTest.h
+++ b/Engine/Source/Developer/CQTest/Public/CQTest.h
@@ -64,9 +64,9 @@ enum class ECQTestSuppressLogBehavior
 	False
 };
 
-static const FString GenerateTestDirectory = TEXT("");
+CQTEST_API extern const FString GenerateTestDirectory;
+CQTEST_API extern const FString DefaultTags;
 static constexpr EAutomationTestFlags DefaultFlags = EAutomationTestFlags_ApplicationContextMask | EAutomationTestFlags::ProductFilter;
-static const FString DefaultTags = TEXT("");
 
 template <typename T>
 concept HasBeforeAll = requires(T t) { { T::BeforeAll(FString()) }; };
diff --git a/Engine/Source/Runtime/Engine/Private/Elements/SMInstance/SMInstanceElementId.cpp b/Engine/Source/Runtime/Engine/Private/Elements/SMInstance/SMInstanceElementId.cpp
--- a/Engine/Source/Runtime/Engine/Private/Elements/SMInstance/SMInstanceElementId.cpp
+++ b/Engine/Source/Runtime/Engine/Private/Elements/SMInstance/SMInstanceElementId.cpp
@@ -12,6 +12,8 @@
 #include "Serialization/TextReferenceCollector.h"
 #endif //WITH_EDITORONLY_DATA
 
+const FName NAME_SMInstance = "SMInstance";
+
 FSMInstanceElementIdMapEntry::FSMInstanceElementIdMapEntry(FSMInstanceElementIdMap* InOwner, UInstancedStaticMeshComponent* InComponent)
 	: Owner(InOwner)
 	, Component(InComponent)
diff --git a/Engine/Source/Runtime/Engine/Public/Elements/SMInstance/SMInstanceElementId.h b/Engine/Source/Runtime/Engine/Public/Elements/SMInstance/SMInstanceElementId.h
--- a/Engine/Source/Runtime/Engine/Public/Elements/SMInstance/SMInstanceElementId.h
+++ b/Engine/Source/Runtime/Engine/Public/Elements/SMInstance/SMInstanceElementId.h
@@ -11,7 +11,7 @@ class FSMInstanceElementIdMap;
 class UInstancedStaticMeshComponent;
 class USMInstanceElementIdMapTransactor;
 
-static const FName NAME_SMInstance = "SMInstance";
+ENGINE_API extern const FName NAME_SMInstance;
 
 /**
  * ID for a specific instance within an ISM, mapped from its instance index.
