[AVCodecs] Fix PVS warning about shifting with a negative number.

diff --git a/Engine/Plugins/Experimental/AVCodecs/AVCodecsCore/Source/AVCodecsCore/Public/Utils/BitstreamReader.h b/Engine/Plugins/Experimental/AVCodecs/AVCodecsCore/Source/AVCodecsCore/Public/Utils/BitstreamReader.h
--- a/Engine/Plugins/Experimental/AVCodecs/AVCodecsCore/Source/AVCodecsCore/Public/Utils/BitstreamReader.h
+++ b/Engine/Plugins/Experimental/AVCodecs/AVCodecsCore/Source/AVCodecsCore/Public/Utils/BitstreamReader.h
@@ -279,7 +279,7 @@ public:
 			b = ReadBits(1);
 		}
 
-		if (lz)
+		if (lz > 0)
 		{
 			Output = static_cast<ValueType>(((1 << lz) | ReadBits(lz)) - 1);
 		}
