The source files included were modified by the UnrealCodeFixup tool so that they can pass the -mergemodules compilation.
*This CL must be submitted before integrating PixelStreaming2 into RemoteSession plugin.

diff --git a/Engine/Plugins/Experimental/AVCodecs/AMFCodecs/Source/AMFCodecs/Public/AMF.h b/Engine/Plugins/Experimental/AVCodecs/AMFCodecs/Source/AMFCodecs/Public/AMF.h
--- a/Engine/Plugins/Experimental/AVCodecs/AMFCodecs/Source/AMFCodecs/Public/AMF.h
+++ b/Engine/Plugins/Experimental/AVCodecs/AMFCodecs/Source/AMFCodecs/Public/AMF.h
@@ -30,14 +30,14 @@ THIRD_PARTY_INCLUDES_START
 
 THIRD_PARTY_INCLUDES_END
 
-class AMFCODECS_API FAMF : public FAPI
+class FAMF : public FAPI
 {
 public:
 	bool bHasCompatibleGPU = true;
 
-	FAMF();
+	AMFCODECS_API FAMF();
 
-	virtual bool IsValid() const override;
+	AMFCODECS_API virtual bool IsValid() const override;
 
 	amf::AMFFactory* GetFactory() const { return Factory; }
 
diff --git a/Engine/Plugins/Experimental/AVCodecs/AMFCodecs/Source/AMFCodecs/Public/Video/Decoders/Configs/VideoDecoderConfigAMF.h b/Engine/Plugins/Experimental/AVCodecs/AMFCodecs/Source/AMFCodecs/Public/Video/Decoders/Configs/VideoDecoderConfigAMF.h
--- a/Engine/Plugins/Experimental/AVCodecs/AMFCodecs/Source/AMFCodecs/Public/Video/Decoders/Configs/VideoDecoderConfigAMF.h
+++ b/Engine/Plugins/Experimental/AVCodecs/AMFCodecs/Source/AMFCodecs/Public/Video/Decoders/Configs/VideoDecoderConfigAMF.h
@@ -7,7 +7,9 @@
 
 #include "AMF.h"
 
-struct AMFCODECS_API FVideoDecoderConfigAMF : public FAVConfig
+#define UE_API AMFCODECS_API
+
+struct FVideoDecoderConfigAMF : public FAVConfig
 {
 private:
 	TMap<wchar_t const*, amf::AMFVariantStruct, FDefaultSetAllocator, TStringPointerMapKeyFuncs_DEPRECATED<wchar_t const*, amf::AMFVariantStruct>> Variants;
@@ -19,8 +21,8 @@ public:
 		TArray<uint8> ExtraData;
 	};
 
-	static FName const CodecTypeH264;
-	static FName const CodecTypeH265;
+	static UE_API FName const CodecTypeH264;
+	static UE_API FName const CodecTypeH265;
 
 	FName CodecType;
 
@@ -108,7 +110,7 @@ public:
 		}
 	}
 
-	FAVResult Parse(TSharedRef<FAVInstance> const& Instance, FVideoPacket const& Packet, TArray<FParsedPicture>& OutPictures);
+	UE_API FAVResult Parse(TSharedRef<FAVInstance> const& Instance, FVideoPacket const& Packet, TArray<FParsedPicture>& OutPictures);
 };
 
 template <>
@@ -124,3 +126,5 @@ template <>
 FAVResult FAVExtension::TransformConfig(FVideoDecoderConfigAMF& OutConfig, struct FVideoDecoderConfigH265 const& InConfig);
 
 DECLARE_TYPEID(FVideoDecoderConfigAMF, AMFCODECS_API);
+
+#undef UE_API
diff --git a/Engine/Plugins/Experimental/AVCodecs/AMFCodecs/Source/AMFCodecs/Public/Video/Encoders/Configs/VideoEncoderConfigAMF.h b/Engine/Plugins/Experimental/AVCodecs/AMFCodecs/Source/AMFCodecs/Public/Video/Encoders/Configs/VideoEncoderConfigAMF.h
--- a/Engine/Plugins/Experimental/AVCodecs/AMFCodecs/Source/AMFCodecs/Public/Video/Encoders/Configs/VideoEncoderConfigAMF.h
+++ b/Engine/Plugins/Experimental/AVCodecs/AMFCodecs/Source/AMFCodecs/Public/Video/Encoders/Configs/VideoEncoderConfigAMF.h
@@ -7,17 +7,19 @@
 
 #include "AMF.h"
 
-struct AMFCODECS_API FVideoEncoderConfigAMF : public FAVConfig
+#define UE_API AMFCODECS_API
+
+struct FVideoEncoderConfigAMF : public FAVConfig
 {
 private:
 	TMap<wchar_t const*, amf::AMFVariantStruct, FDefaultSetAllocator, TStringPointerMapKeyFuncs_DEPRECATED<wchar_t const*, amf::AMFVariantStruct>> Variants;
 
 public:
-	static TAVResult<AMF_VIDEO_ENCODER_RATE_CONTROL_METHOD_ENUM> ConvertRateControlMode(ERateControlMode Mode);
-	static TAVResult<ERateControlMode> ConvertRateControlMode(AMF_VIDEO_ENCODER_RATE_CONTROL_METHOD_ENUM Mode);
+	static UE_API TAVResult<AMF_VIDEO_ENCODER_RATE_CONTROL_METHOD_ENUM> ConvertRateControlMode(ERateControlMode Mode);
+	static UE_API TAVResult<ERateControlMode> ConvertRateControlMode(AMF_VIDEO_ENCODER_RATE_CONTROL_METHOD_ENUM Mode);
 
-	static FName const CodecTypeH264;
-	static FName const CodecTypeH265;
+	static UE_API FName const CodecTypeH264;
+	static UE_API FName const CodecTypeH265;
 
 	FName CodecType;
 
@@ -169,3 +171,5 @@ FAVResult FAVExtension::TransformConfig(FVideoEncoderConfigAMF& OutConfig, struc
 // FAVResult FAVExtension::TransformConfig(FVideoEncoderConfigAMF& OutConfig, struct FVideoEncoderConfigH265 const& InConfig);
 
 DECLARE_TYPEID(FVideoEncoderConfigAMF, AMFCODECS_API);
+
+#undef UE_API
diff --git a/Engine/Plugins/Experimental/AVCodecs/AVCodecsCore/Source/AVCodecsCore/Private/Video/CodecUtils/CodecUtilsH264.cpp b/Engine/Plugins/Experimental/AVCodecs/AVCodecsCore/Source/AVCodecsCore/Private/Video/CodecUtils/CodecUtilsH264.cpp
--- a/Engine/Plugins/Experimental/AVCodecs/AVCodecsCore/Source/AVCodecsCore/Private/Video/CodecUtils/CodecUtilsH264.cpp
+++ b/Engine/Plugins/Experimental/AVCodecs/AVCodecsCore/Source/AVCodecsCore/Private/Video/CodecUtils/CodecUtilsH264.cpp
@@ -6,6 +6,8 @@
 #include "AVResult.h"
 #include "Video/VideoPacket.h"
 
+#include UE_INLINE_GENERATED_CPP_BY_NAME(CodecUtilsH264)
+
 FH264ProfileDefinition GH264ProfileDefinitions[static_cast<uint8>(EH264Profile::MAX)] = {
 	{ EH264Profile::Auto, UE::AVCodecCore::H264::EH264ProfileIDC::Auto, UE::AVCodecCore::H264::EH264ConstraintFlag::None, TEXT("Auto") },
 	{ EH264Profile::CALVLC444Intra, UE::AVCodecCore::H264::EH264ProfileIDC::CALVLC444Intra, UE::AVCodecCore::H264::EH264ConstraintFlag::None, TEXT("CALVLC444Intra") },
@@ -616,4 +618,4 @@ namespace UE::AVCodecCore::H264
 		return EAVResult::Success;
 	}
 
-} // namespace UE::AVCodecCore::H264
\ No newline at end of file
+} // namespace UE::AVCodecCore::H264
diff --git a/Engine/Plugins/Experimental/AVCodecs/AVCodecsCore/Source/AVCodecsCore/Private/Video/VideoEncoder.cpp b/Engine/Plugins/Experimental/AVCodecs/AVCodecsCore/Source/AVCodecsCore/Private/Video/VideoEncoder.cpp
--- a/Engine/Plugins/Experimental/AVCodecs/AVCodecsCore/Source/AVCodecsCore/Private/Video/VideoEncoder.cpp
+++ b/Engine/Plugins/Experimental/AVCodecs/AVCodecsCore/Source/AVCodecsCore/Private/Video/VideoEncoder.cpp
@@ -2,6 +2,8 @@
 
 #include "Video/VideoEncoder.h"
 
+#include UE_INLINE_GENERATED_CPP_BY_NAME(VideoEncoder)
+
 REGISTER_TYPEID(FVideoEncoder);
 
 TMap<FString, EScalabilityMode> const ScalabilityModeMap = {
@@ -40,4 +42,4 @@ TMap<FString, EScalabilityMode> const ScalabilityModeMap = {
 	{ "S3T3", EScalabilityMode::S3T3 },
 	{ "S3T3h", EScalabilityMode::S3T3h },
 	{ "None", EScalabilityMode::None }
-};
\ No newline at end of file
+};
diff --git a/Engine/Plugins/Experimental/AVCodecs/AVCodecsCore/Source/AVCodecsCore/Public/AVCoder.h b/Engine/Plugins/Experimental/AVCodecs/AVCodecsCore/Source/AVCodecsCore/Public/AVCoder.h
--- a/Engine/Plugins/Experimental/AVCodecs/AVCodecsCore/Source/AVCodecsCore/Public/AVCoder.h
+++ b/Engine/Plugins/Experimental/AVCodecs/AVCodecsCore/Source/AVCodecsCore/Public/AVCoder.h
@@ -8,6 +8,8 @@
 #include "AVResource.h"
 #include "AVResult.h"
 
+#define UE_API AVCODECSCORE_API
+
 /**
  * HOW TO USE
  *
@@ -43,7 +45,7 @@ class TAVCoder;
 /**
  * Simple base coder interface
  */
-class AVCODECSCORE_API IAVCoder
+class IAVCoder
 {
 protected:
 	/**
@@ -54,7 +56,7 @@ protected:
 	 * @see FTypeID
 	 * @see TAVCoder::GetFactories
 	 */
-	static TMap<FTypeID, TMap<FTypeID, TMap<FTypeID, TSharedPtr<void>>>> Factories;
+	static UE_API TMap<FTypeID, TMap<FTypeID, TMap<FTypeID, TSharedPtr<void>>>> Factories;
 
 public:
 	virtual ~IAVCoder() = default;
@@ -608,3 +610,5 @@ public:
 		return EAVResult::Success;
 	}
 };
+
+#undef UE_API
diff --git a/Engine/Plugins/Experimental/AVCodecs/AVCodecsCore/Source/AVCodecsCore/Public/AVDevice.h b/Engine/Plugins/Experimental/AVCodecs/AVCodecsCore/Source/AVCodecsCore/Public/AVDevice.h
--- a/Engine/Plugins/Experimental/AVCodecs/AVCodecsCore/Source/AVCodecsCore/Public/AVDevice.h
+++ b/Engine/Plugins/Experimental/AVCodecs/AVCodecsCore/Source/AVCodecsCore/Public/AVDevice.h
@@ -5,10 +5,12 @@
 #include "AVContext.h"
 #include "AVUtility.h"
 
+#define UE_API AVCODECSCORE_API
+
 /**
  * A representation of a physical computer device, be it a GPU or a CPU or an external device, that can hold device resource contexts (Vulkan, D3D, CUDA, etc).
  */
-class AVCODECSCORE_API FAVDevice
+class FAVDevice
 {
 private:
 	/*
@@ -23,14 +25,14 @@ public:
 	 * @param Index Index of the device if there are multiple hardware devices present (multiple GPUs or external cards).
 	 * @return The global hardware device at Index.
 	 */
-	static TSharedRef<FAVDevice>& GetHardwareDevice(int32 Index = 0);
+	static UE_API TSharedRef<FAVDevice>& GetHardwareDevice(int32 Index = 0);
 
 	/**
 	 * Get the global software device.
 	 *
 	 * @return The global software device.
 	 */
-	static TSharedRef<FAVDevice>& GetSoftwareDevice();
+	static UE_API TSharedRef<FAVDevice>& GetSoftwareDevice();
 
 	/**
 	 * Check if this device has a specific context.
@@ -71,3 +73,5 @@ public:
 	FAVDevice() = default;
 	virtual ~FAVDevice() = default;
 };
+
+#undef UE_API
diff --git a/Engine/Plugins/Experimental/AVCodecs/AVCodecsCore/Source/AVCodecsCore/Public/AVExtension.h b/Engine/Plugins/Experimental/AVCodecs/AVCodecsCore/Source/AVCodecsCore/Public/AVExtension.h
--- a/Engine/Plugins/Experimental/AVCodecs/AVCodecsCore/Source/AVCodecsCore/Public/AVExtension.h
+++ b/Engine/Plugins/Experimental/AVCodecs/AVCodecsCore/Source/AVCodecsCore/Public/AVExtension.h
@@ -32,7 +32,7 @@
 /**
  * Static extension system for simple composition of different resource/config types, via exported specializations
  */
-struct AVCODECSCORE_API FAVExtension
+struct FAVExtension
 {
 public:
 	/**
diff --git a/Engine/Plugins/Experimental/AVCodecs/AVCodecsCore/Source/AVCodecsCore/Public/AVResource.h b/Engine/Plugins/Experimental/AVCodecs/AVCodecsCore/Source/AVCodecsCore/Public/AVResource.h
--- a/Engine/Plugins/Experimental/AVCodecs/AVCodecsCore/Source/AVCodecsCore/Public/AVResource.h
+++ b/Engine/Plugins/Experimental/AVCodecs/AVCodecsCore/Source/AVCodecsCore/Public/AVResource.h
@@ -6,11 +6,13 @@
 #include "AVLayout.h"
 #include "AVResult.h"
 
+#define UE_API AVCODECSCORE_API
+
 /**
  * Base wrapper for an immutable device resource. Bundles necessary data that may not be easily retrieved on without an RHI, such as size on device.
  * If the underlying device resource must change, this resource wrapper must be recreated.
  */
-class AVCODECSCORE_API FAVResource
+class FAVResource
 {
 private:
 	/**
@@ -60,7 +62,7 @@ public:
 	 */
 	inline uint64 GetStride() const { return Layout.Stride; }
 
-	FAVResource(TSharedRef<FAVDevice> const& Device, FAVLayout const& Layout);
+	UE_API FAVResource(TSharedRef<FAVDevice> const& Device, FAVLayout const& Layout);
 	virtual ~FAVResource() = default;
 
 	/**
@@ -68,24 +70,24 @@ public:
 	 *
 	 * @return Result of the operation, @see FAVResult.
 	 */
-	virtual FAVResult Validate() const;
+	UE_API virtual FAVResult Validate() const;
 
 	/**
 	 * Lock this resource for internal use. Is a blocking call until the lock can be taken.
 	 */
-	virtual void Lock();
+	UE_API virtual void Lock();
 	
 	/**
 	 * Lock this resource for internal use, within the current scope. Is a blocking call until the lock can be taken.
 	 *
 	 * @return The lock object. Capture this, because when it goes out of scope it will unlock the resource.
 	 */
-	virtual FScopeLock LockScope();
+	UE_API virtual FScopeLock LockScope();
 
 	/**
 	 * Release this resource for external use.
 	 */
-	virtual void Unlock();
+	UE_API virtual void Unlock();
 
 	/**
 	 * Create a mapping of this resource to another resource (or other object). Returned reference must then be set by the caller.
@@ -99,3 +101,5 @@ public:
 		return Mappings.Edit<TMapping>();
 	}
 };
+
+#undef UE_API
diff --git a/Engine/Plugins/Experimental/AVCodecs/AVCodecsCore/Source/AVCodecsCore/Public/AVUtility.h b/Engine/Plugins/Experimental/AVCodecs/AVCodecsCore/Source/AVCodecsCore/Public/AVUtility.h
--- a/Engine/Plugins/Experimental/AVCodecs/AVCodecsCore/Source/AVCodecsCore/Public/AVUtility.h
+++ b/Engine/Plugins/Experimental/AVCodecs/AVCodecsCore/Source/AVCodecsCore/Public/AVUtility.h
@@ -4,6 +4,8 @@
 
 #include "CoreMinimal.h"
 
+#define UE_API AVCODECSCORE_API
+
 /**
  * All generic (non-AV) utility types.
  */
@@ -304,13 +306,13 @@ public:
 /**
  * API abstraction with static singletons.
  */
-class AVCODECSCORE_API FAPI
+class FAPI
 {
 private:
 	/**
 	 * Type-erased map of API singletons.
 	 */
-	static TTypeMap<FAPI> Singletons;
+	static UE_API TTypeMap<FAPI> Singletons;
 
 public:
 	/**
@@ -353,3 +355,5 @@ protected:
 	FAPI() = default;
 	virtual ~FAPI() = default;
 };
+
+#undef UE_API
diff --git a/Engine/Plugins/Experimental/AVCodecs/AVCodecsCore/Source/AVCodecsCore/Public/Audio/AudioResource.h b/Engine/Plugins/Experimental/AVCodecs/AVCodecsCore/Source/AVCodecsCore/Public/Audio/AudioResource.h
--- a/Engine/Plugins/Experimental/AVCodecs/AVCodecsCore/Source/AVCodecsCore/Public/Audio/AudioResource.h
+++ b/Engine/Plugins/Experimental/AVCodecs/AVCodecsCore/Source/AVCodecsCore/Public/Audio/AudioResource.h
@@ -6,6 +6,8 @@
 #include "AVResource.h"
 #include "AVUtility.h"
 
+#define UE_API AVCODECSCORE_API
+
 // enum class EAudioFormat : uint8
 // {
 // 	// TODO (Andrew)
@@ -69,7 +71,7 @@ public:
 /**
  * Base wrapper for a audio device resource.
  */
-class AVCODECSCORE_API FAudioResource : public FAVResource
+class FAudioResource : public FAVResource
 {
 private:
 	/**
@@ -93,7 +95,7 @@ public:
 	 */
 	inline float GetSampleDuration() const { return Descriptor.SampleDuration; }
 
-	FAudioResource(TSharedRef<FAVDevice> const& Device, FAVLayout const& Layout, FAudioDescriptor const& Descriptor);
+	UE_API FAudioResource(TSharedRef<FAVDevice> const& Device, FAVLayout const& Layout, FAudioDescriptor const& Descriptor);
 	virtual ~FAudioResource() override = default;
 };
 
@@ -152,3 +154,5 @@ protected:
 		return nullptr;
 	}
 };
+
+#undef UE_API
diff --git a/Engine/Plugins/Experimental/AVCodecs/AVCodecsCore/Source/AVCodecsCore/Public/Audio/Resources/AudioResourceCPU.h b/Engine/Plugins/Experimental/AVCodecs/AVCodecsCore/Source/AVCodecsCore/Public/Audio/Resources/AudioResourceCPU.h
--- a/Engine/Plugins/Experimental/AVCodecs/AVCodecsCore/Source/AVCodecsCore/Public/Audio/Resources/AudioResourceCPU.h
+++ b/Engine/Plugins/Experimental/AVCodecs/AVCodecsCore/Source/AVCodecsCore/Public/Audio/Resources/AudioResourceCPU.h
@@ -5,17 +5,19 @@
 #include "AVContext.h"
 #include "Audio/AudioResource.h"
 
+#define UE_API AVCODECSCORE_API
+
 /**
  * CPU audio context and resource.
  */
 
-class AVCODECSCORE_API FAudioContextCPU : public FAVContext
+class FAudioContextCPU : public FAVContext
 {
 public:
-	FAudioContextCPU();
+	UE_API FAudioContextCPU();
 };
 
-class AVCODECSCORE_API FAudioResourceCPU : public TAudioResource<FAudioContextCPU>
+class FAudioResourceCPU : public TAudioResource<FAudioContextCPU>
 {
 private:
 	TSharedPtr<float> Raw;
@@ -23,11 +25,13 @@ private:
 public:
 	inline TSharedPtr<float> const& GetRaw() const { return Raw; }
 
-	FAudioResourceCPU(TSharedRef<FAVDevice> const& Device, TSharedPtr<float> const& Raw, FAVLayout const& Layout, FAudioDescriptor const& Descriptor);
+	UE_API FAudioResourceCPU(TSharedRef<FAVDevice> const& Device, TSharedPtr<float> const& Raw, FAVLayout const& Layout, FAudioDescriptor const& Descriptor);
 	virtual ~FAudioResourceCPU() override = default;
 
-	virtual FAVResult Validate() const override;
+	UE_API virtual FAVResult Validate() const override;
 };
 
 DECLARE_TYPEID(FAudioContextCPU, AVCODECSCORE_API);
 DECLARE_TYPEID(FAudioResourceCPU, AVCODECSCORE_API);
+
+#undef UE_API
diff --git a/Engine/Plugins/Experimental/AVCodecs/AVCodecsCore/Source/AVCodecsCore/Public/Video/Decoders/Configs/VideoDecoderConfigAV1.h b/Engine/Plugins/Experimental/AVCodecs/AVCodecsCore/Source/AVCodecsCore/Public/Video/Decoders/Configs/VideoDecoderConfigAV1.h
--- a/Engine/Plugins/Experimental/AVCodecs/AVCodecsCore/Source/AVCodecsCore/Public/Video/Decoders/Configs/VideoDecoderConfigAV1.h
+++ b/Engine/Plugins/Experimental/AVCodecs/AVCodecsCore/Source/AVCodecsCore/Public/Video/Decoders/Configs/VideoDecoderConfigAV1.h
@@ -8,7 +8,7 @@
 /*
  * Configuration settings for AV1 decoders.
  */
-struct AVCODECSCORE_API FVideoDecoderConfigAV1 : public FVideoDecoderConfig
+struct FVideoDecoderConfigAV1 : public FVideoDecoderConfig
 {
     FVideoDecoderConfigAV1(EAVPreset Preset = EAVPreset::Default)
 		: FVideoDecoderConfig(Preset)
diff --git a/Engine/Plugins/Experimental/AVCodecs/AVCodecsCore/Source/AVCodecsCore/Public/Video/Decoders/Configs/VideoDecoderConfigH264.h b/Engine/Plugins/Experimental/AVCodecs/AVCodecsCore/Source/AVCodecsCore/Public/Video/Decoders/Configs/VideoDecoderConfigH264.h
--- a/Engine/Plugins/Experimental/AVCodecs/AVCodecsCore/Source/AVCodecsCore/Public/Video/Decoders/Configs/VideoDecoderConfigH264.h
+++ b/Engine/Plugins/Experimental/AVCodecs/AVCodecsCore/Source/AVCodecsCore/Public/Video/Decoders/Configs/VideoDecoderConfigH264.h
@@ -5,10 +5,12 @@
 #include "Video/VideoDecoder.h"
 #include "Video/CodecUtils/CodecUtilsH264.h"
 
+#define UE_API AVCODECSCORE_API
+
 /*
  * Configuration settings for H264 decoders.
  */
-struct AVCODECSCORE_API FVideoDecoderConfigH264 : public FVideoDecoderConfig
+struct FVideoDecoderConfigH264 : public FVideoDecoderConfig
 {
 	TMap<uint32, UE::AVCodecCore::H264::SPS_t> SPS;
 	TMap<uint32, UE::AVCodecCore::H264::PPS_t> PPS;
@@ -19,9 +21,11 @@ struct AVCODECSCORE_API FVideoDecoderConfigH264 : public FVideoDecoderConfig
 	{
 	}
 
-	FAVResult Parse(FVideoPacket const& Packet, TArray<UE::AVCodecCore::H264::Slice_t>& OutSlices);
+	UE_API FAVResult Parse(FVideoPacket const& Packet, TArray<UE::AVCodecCore::H264::Slice_t>& OutSlices);
 
-    TOptional<int> GetLastSliceQP(TArray<UE::AVCodecCore::H264::Slice_t>& Slices);
+    UE_API TOptional<int> GetLastSliceQP(TArray<UE::AVCodecCore::H264::Slice_t>& Slices);
 };
 
 DECLARE_TYPEID(FVideoDecoderConfigH264, AVCODECSCORE_API);
+
+#undef UE_API
diff --git a/Engine/Plugins/Experimental/AVCodecs/AVCodecsCore/Source/AVCodecsCore/Public/Video/Decoders/Configs/VideoDecoderConfigH265.h b/Engine/Plugins/Experimental/AVCodecs/AVCodecsCore/Source/AVCodecsCore/Public/Video/Decoders/Configs/VideoDecoderConfigH265.h
--- a/Engine/Plugins/Experimental/AVCodecs/AVCodecsCore/Source/AVCodecsCore/Public/Video/Decoders/Configs/VideoDecoderConfigH265.h
+++ b/Engine/Plugins/Experimental/AVCodecs/AVCodecsCore/Source/AVCodecsCore/Public/Video/Decoders/Configs/VideoDecoderConfigH265.h
@@ -5,10 +5,12 @@
 #include "Video/VideoDecoder.h"
 #include "Video/CodecUtils/CodecUtilsH265.h"
 
+#define UE_API AVCODECSCORE_API
+
 /*
  * Configuration settings for H265 decoders.
  */
-struct AVCODECSCORE_API FVideoDecoderConfigH265 : public FVideoDecoderConfig
+struct FVideoDecoderConfigH265 : public FVideoDecoderConfig
 {
     TMap<uint32, UE::AVCodecCore::H265::VPS_t> VPS;
 	TMap<uint32, UE::AVCodecCore::H265::SPS_t> SPS;
@@ -20,9 +22,11 @@ struct AVCODECSCORE_API FVideoDecoderConfigH265 : public FVideoDecoderConfig
 	{
 	}
 
-	FAVResult Parse(FVideoPacket const& Packet, TArray<UE::AVCodecCore::H265::Slice_t>& OutSlices);
+	UE_API FAVResult Parse(FVideoPacket const& Packet, TArray<UE::AVCodecCore::H265::Slice_t>& OutSlices);
 
-    TOptional<int> GetLastSliceQP(TArray<UE::AVCodecCore::H265::Slice_t>& Slices);
+    UE_API TOptional<int> GetLastSliceQP(TArray<UE::AVCodecCore::H265::Slice_t>& Slices);
 };
 
 DECLARE_TYPEID(FVideoDecoderConfigH265, AVCODECSCORE_API);
+
+#undef UE_API
diff --git a/Engine/Plugins/Experimental/AVCodecs/AVCodecsCore/Source/AVCodecsCore/Public/Video/Decoders/Configs/VideoDecoderConfigVP8.h b/Engine/Plugins/Experimental/AVCodecs/AVCodecsCore/Source/AVCodecsCore/Public/Video/Decoders/Configs/VideoDecoderConfigVP8.h
--- a/Engine/Plugins/Experimental/AVCodecs/AVCodecsCore/Source/AVCodecsCore/Public/Video/Decoders/Configs/VideoDecoderConfigVP8.h
+++ b/Engine/Plugins/Experimental/AVCodecs/AVCodecsCore/Source/AVCodecsCore/Public/Video/Decoders/Configs/VideoDecoderConfigVP8.h
@@ -7,7 +7,7 @@
 /*
  * Configuration settings for VP8 decoders.
  */
-struct AVCODECSCORE_API FVideoDecoderConfigVP8 : public FVideoDecoderConfig
+struct FVideoDecoderConfigVP8 : public FVideoDecoderConfig
 {
 	int32 NumberOfCores = 0;
 
diff --git a/Engine/Plugins/Experimental/AVCodecs/AVCodecsCore/Source/AVCodecsCore/Public/Video/Decoders/Configs/VideoDecoderConfigVP9.h b/Engine/Plugins/Experimental/AVCodecs/AVCodecsCore/Source/AVCodecsCore/Public/Video/Decoders/Configs/VideoDecoderConfigVP9.h
--- a/Engine/Plugins/Experimental/AVCodecs/AVCodecsCore/Source/AVCodecsCore/Public/Video/Decoders/Configs/VideoDecoderConfigVP9.h
+++ b/Engine/Plugins/Experimental/AVCodecs/AVCodecsCore/Source/AVCodecsCore/Public/Video/Decoders/Configs/VideoDecoderConfigVP9.h
@@ -5,10 +5,12 @@
 #include "Video/VideoDecoder.h"
 #include "Video/CodecUtils/CodecUtilsVP9.h"
 
+#define UE_API AVCODECSCORE_API
+
 /*
  * Configuration settings for VP9 decoders.
  */
-struct AVCODECSCORE_API FVideoDecoderConfigVP9 : public FVideoDecoderConfig
+struct FVideoDecoderConfigVP9 : public FVideoDecoderConfig
 {
 	// Values are parsed from the bitstream during decoding
 	uint32 MaxOutputWidth = 0;
@@ -21,7 +23,9 @@ struct AVCODECSCORE_API FVideoDecoderConfigVP9 : public FVideoDecoderConfig
 	{
 	}
 
-	FAVResult Parse(FVideoPacket const& Packet, UE::AVCodecCore::VP9::Header_t& OutHeader);
+	UE_API FAVResult Parse(FVideoPacket const& Packet, UE::AVCodecCore::VP9::Header_t& OutHeader);
 };
 
 DECLARE_TYPEID(FVideoDecoderConfigVP9, AVCODECSCORE_API);
+
+#undef UE_API
diff --git a/Engine/Plugins/Experimental/AVCodecs/AVCodecsCore/Source/AVCodecsCore/Public/Video/DependencyDescriptor.h b/Engine/Plugins/Experimental/AVCodecs/AVCodecsCore/Source/AVCodecsCore/Public/Video/DependencyDescriptor.h
--- a/Engine/Plugins/Experimental/AVCodecs/AVCodecsCore/Source/AVCodecsCore/Public/Video/DependencyDescriptor.h
+++ b/Engine/Plugins/Experimental/AVCodecs/AVCodecsCore/Source/AVCodecsCore/Public/Video/DependencyDescriptor.h
@@ -6,6 +6,8 @@
 #include "Containers/UnrealString.h"
 #include "Math/IntPoint.h"
 
+#define UE_API AVCODECSCORE_API
+
 // Relationship of a frame to a Decode target.
 enum class AVCODECSCORE_API EDecodeTargetIndication : uint8
 {
@@ -15,15 +17,15 @@ enum class AVCODECSCORE_API EDecodeTargetIndication : uint8
 	Required = 3	 // DecodeTargetInfo symbol 'R'
 };
 
-class AVCODECSCORE_API FFrameDependencyTemplate
+class FFrameDependencyTemplate
 {
 public:
 	// Setters are named briefly to chain them when building the template.
-	FFrameDependencyTemplate& SpatialLayerId(int SpatialLayer);
-	FFrameDependencyTemplate& TemporalLayerId(int TemporalLayer);
-	FFrameDependencyTemplate& Dtis(const FString& DTIS);
-	FFrameDependencyTemplate& FrameDiff(std::initializer_list<int32> Diffs);
-	FFrameDependencyTemplate& ChainDiff(std::initializer_list<int32> Diffs);
+	UE_API FFrameDependencyTemplate& SpatialLayerId(int SpatialLayer);
+	UE_API FFrameDependencyTemplate& TemporalLayerId(int TemporalLayer);
+	UE_API FFrameDependencyTemplate& Dtis(const FString& DTIS);
+	UE_API FFrameDependencyTemplate& FrameDiff(std::initializer_list<int32> Diffs);
+	UE_API FFrameDependencyTemplate& ChainDiff(std::initializer_list<int32> Diffs);
 
 	friend bool operator==(const FFrameDependencyTemplate& Lhs, const FFrameDependencyTemplate& Rhs)
 	{
@@ -41,7 +43,7 @@ public:
 	TArray<int32>					ChainDiffs;
 };
 
-class AVCODECSCORE_API FFrameDependencyStructure
+class FFrameDependencyStructure
 {
 public:
 	friend bool operator==(const FFrameDependencyStructure& Lhs, const FFrameDependencyStructure& Rhs)
@@ -61,4 +63,6 @@ public:
 	TArray<int32>					 DecodeTargetProtectedByChain;
 	TArray<FIntPoint>				 Resolutions;
 	TArray<FFrameDependencyTemplate> Templates;
-};
\ No newline at end of file
+};
+
+#undef UE_API
diff --git a/Engine/Plugins/Experimental/AVCodecs/AVCodecsCore/Source/AVCodecsCore/Public/Video/Encoders/SVC/ScalabilityStructureFull.h b/Engine/Plugins/Experimental/AVCodecs/AVCodecsCore/Source/AVCodecsCore/Public/Video/Encoders/SVC/ScalabilityStructureFull.h
--- a/Engine/Plugins/Experimental/AVCodecs/AVCodecsCore/Source/AVCodecsCore/Public/Video/Encoders/SVC/ScalabilityStructureFull.h
+++ b/Engine/Plugins/Experimental/AVCodecs/AVCodecsCore/Source/AVCodecsCore/Public/Video/Encoders/SVC/ScalabilityStructureFull.h
@@ -5,17 +5,19 @@
 #include "Video/GenericFrameInfo.h"
 #include "Video/Encoders/SVC/ScalableVideoController.h"
 
-class AVCODECSCORE_API FScalabilityStructureFullSvc : public FScalableVideoController
+#define UE_API AVCODECSCORE_API
+
+class FScalabilityStructureFullSvc : public FScalableVideoController
 {
 public:
-	FScalabilityStructureFullSvc(int NumSpatialLayers, int NumTemporalLayers, FIntFraction ResolutionFactor);
+	UE_API FScalabilityStructureFullSvc(int NumSpatialLayers, int NumTemporalLayers, FIntFraction ResolutionFactor);
 	~FScalabilityStructureFullSvc() override = default;
 
-	virtual FStreamLayersConfig StreamConfig() const override;
+	UE_API virtual FStreamLayersConfig StreamConfig() const override;
 
-	virtual TArray<FScalableVideoController::FLayerFrameConfig> NextFrameConfig(bool bRestart) override;
-	virtual FGenericFrameInfo									OnEncodeDone(const FScalableVideoController::FLayerFrameConfig& Config) override;
-	virtual void												OnRatesUpdated(const FVideoBitrateAllocation& Bitrates) override;
+	UE_API virtual TArray<FScalableVideoController::FLayerFrameConfig> NextFrameConfig(bool bRestart) override;
+	UE_API virtual FGenericFrameInfo									OnEncodeDone(const FScalableVideoController::FLayerFrameConfig& Config) override;
+	UE_API virtual void												OnRatesUpdated(const FVideoBitrateAllocation& Bitrates) override;
 
 private:
 	// NOTE: This enum name is duplicated throughout the ScalabilityStructure variants
@@ -30,7 +32,7 @@ private:
 		DeltaT0,
 	};
 
-	static FString		 FramePatternNames[6];
+	static UE_API FString		 FramePatternNames[6];
 	static constexpr int MaxNumSpatialLayers = 3;
 	static constexpr int MaxNumTemporalLayers = 3;
 
@@ -50,9 +52,9 @@ private:
 		ActiveDecodeTargets[Sid * NumTemporalLayers + Tid] = Value;
 	}
 
-	EFramePattern				   NextPattern() const;
-	bool						   TemporalLayerIsActive(int Tid) const;
-	static EDecodeTargetIndication Dti(int Sid, int Tid, const FScalableVideoController::FLayerFrameConfig& Frame);
+	UE_API EFramePattern				   NextPattern() const;
+	UE_API bool						   TemporalLayerIsActive(int Tid) const;
+	static UE_API EDecodeTargetIndication Dti(int Sid, int Tid, const FScalableVideoController::FLayerFrameConfig& Frame);
 
 	const int			 NumSpatialLayers;
 	const int			 NumTemporalLayers;
@@ -68,7 +70,7 @@ private:
 //         /   /   / ...
 // T0     0---0---0--
 // Time-> 0 1 2 3 4
-class AVCODECSCORE_API FScalabilityStructureL1T2 : public FScalabilityStructureFullSvc
+class FScalabilityStructureL1T2 : public FScalabilityStructureFullSvc
 {
 public:
 	explicit FScalabilityStructureL1T2(FIntFraction ResolutionFactor = { 1, 2 })
@@ -77,7 +79,7 @@ public:
 	}
 	~FScalabilityStructureL1T2() override = default;
 
-	FFrameDependencyStructure DependencyStructure() const override;
+	UE_API FFrameDependencyStructure DependencyStructure() const override;
 };
 
 // T2       0   0   0   0
@@ -86,7 +88,7 @@ public:
 //         |_/     |_/
 // T0     0-------0------
 // Time-> 0 1 2 3 4 5 6 7
-class AVCODECSCORE_API FScalabilityStructureL1T3 : public FScalabilityStructureFullSvc
+class FScalabilityStructureL1T3 : public FScalabilityStructureFullSvc
 {
 public:
 	explicit FScalabilityStructureL1T3(FIntFraction ResolutionFactor = { 1, 2 })
@@ -95,13 +97,13 @@ public:
 	}
 	~FScalabilityStructureL1T3() override = default;
 
-	FFrameDependencyStructure DependencyStructure() const override;
+	UE_API FFrameDependencyStructure DependencyStructure() const override;
 };
 
 // S1  0--0--0-
 //     |  |  | ...
 // S0  0--0--0-
-class AVCODECSCORE_API FScalabilityStructureL2T1 : public FScalabilityStructureFullSvc
+class FScalabilityStructureL2T1 : public FScalabilityStructureFullSvc
 {
 public:
 	explicit FScalabilityStructureL2T1(FIntFraction ResolutionFactor = { 1, 2 })
@@ -110,7 +112,7 @@ public:
 	}
 	~FScalabilityStructureL2T1() override = default;
 
-	FFrameDependencyStructure DependencyStructure() const override;
+	UE_API FFrameDependencyStructure DependencyStructure() const override;
 };
 
 // S1T1     0   0
@@ -121,7 +123,7 @@ public:
 //        |/  |/  |/
 // S0T0   0---0---0--
 // Time-> 0 1 2 3 4
-class AVCODECSCORE_API FScalabilityStructureL2T2 : public FScalabilityStructureFullSvc
+class FScalabilityStructureL2T2 : public FScalabilityStructureFullSvc
 {
 public:
 	explicit FScalabilityStructureL2T2(FIntFraction ResolutionFactor = { 1, 2 })
@@ -130,7 +132,7 @@ public:
 	}
 	~FScalabilityStructureL2T2() override = default;
 
-	FFrameDependencyStructure DependencyStructure() const override;
+	UE_API FFrameDependencyStructure DependencyStructure() const override;
 };
 
 // S1T2      4    ,8
@@ -141,7 +143,7 @@ public:
 // S0T1   | /   5'
 // S0T0   1----'-----...
 // Time-> 0  1  2  3
-class AVCODECSCORE_API FScalabilityStructureL2T3 : public FScalabilityStructureFullSvc
+class FScalabilityStructureL2T3 : public FScalabilityStructureFullSvc
 {
 public:
 	explicit FScalabilityStructureL2T3(FIntFraction ResolutionFactor = { 1, 2 })
@@ -150,7 +152,7 @@ public:
 	}
 	~FScalabilityStructureL2T3() override = default;
 
-	FFrameDependencyStructure DependencyStructure() const override;
+	UE_API FFrameDependencyStructure DependencyStructure() const override;
 };
 
 // S2     0-0-0-
@@ -159,7 +161,7 @@ public:
 //        | | |
 // S0     0-0-0-
 // Time-> 0 1 2
-class AVCODECSCORE_API FScalabilityStructureL3T1 : public FScalabilityStructureFullSvc
+class FScalabilityStructureL3T1 : public FScalabilityStructureFullSvc
 {
 public:
 	explicit FScalabilityStructureL3T1(FIntFraction ResolutionFactor = { 1, 2 })
@@ -168,11 +170,11 @@ public:
 	}
 	~FScalabilityStructureL3T1() override = default;
 
-	FFrameDependencyStructure DependencyStructure() const override;
+	UE_API FFrameDependencyStructure DependencyStructure() const override;
 };
 
 // https://www.w3.org/TR/webrtc-svc/#L3T2*
-class AVCODECSCORE_API FScalabilityStructureL3T2 : public FScalabilityStructureFullSvc
+class FScalabilityStructureL3T2 : public FScalabilityStructureFullSvc
 {
 public:
 	explicit FScalabilityStructureL3T2(FIntFraction ResolutionFactor = { 1, 2 })
@@ -181,11 +183,11 @@ public:
 	}
 	~FScalabilityStructureL3T2() override = default;
 
-	FFrameDependencyStructure DependencyStructure() const override;
+	UE_API FFrameDependencyStructure DependencyStructure() const override;
 };
 
 // https://www.w3.org/TR/webrtc-svc/#L3T3*
-class AVCODECSCORE_API FScalabilityStructureL3T3 : public FScalabilityStructureFullSvc
+class FScalabilityStructureL3T3 : public FScalabilityStructureFullSvc
 {
 public:
 	explicit FScalabilityStructureL3T3(FIntFraction ResolutionFactor = { 1, 2 })
@@ -194,5 +196,7 @@ public:
 	}
 	~FScalabilityStructureL3T3() override = default;
 
-	FFrameDependencyStructure DependencyStructure() const override;
-};
\ No newline at end of file
+	UE_API FFrameDependencyStructure DependencyStructure() const override;
+};
+
+#undef UE_API
diff --git a/Engine/Plugins/Experimental/AVCodecs/AVCodecsCore/Source/AVCodecsCore/Public/Video/Encoders/SVC/ScalabilityStructureKey.h b/Engine/Plugins/Experimental/AVCodecs/AVCodecsCore/Source/AVCodecsCore/Public/Video/Encoders/SVC/ScalabilityStructureKey.h
--- a/Engine/Plugins/Experimental/AVCodecs/AVCodecsCore/Source/AVCodecsCore/Public/Video/Encoders/SVC/ScalabilityStructureKey.h
+++ b/Engine/Plugins/Experimental/AVCodecs/AVCodecsCore/Source/AVCodecsCore/Public/Video/Encoders/SVC/ScalabilityStructureKey.h
@@ -5,17 +5,19 @@
 #include "Video/GenericFrameInfo.h"
 #include "Video/Encoders/SVC/ScalableVideoController.h"
 
-class AVCODECSCORE_API FScalabilityStructureKeySvc : public FScalableVideoController
+#define UE_API AVCODECSCORE_API
+
+class FScalabilityStructureKeySvc : public FScalableVideoController
 {
 public:
-	FScalabilityStructureKeySvc(int NumSpatialLayers, int NumTemporalLayers);
+	UE_API FScalabilityStructureKeySvc(int NumSpatialLayers, int NumTemporalLayers);
 	~FScalabilityStructureKeySvc() override = default;
 
-	virtual FStreamLayersConfig StreamConfig() const override;
+	UE_API virtual FStreamLayersConfig StreamConfig() const override;
 
-	virtual TArray<FScalableVideoController::FLayerFrameConfig> NextFrameConfig(bool bRestart) override;
-	virtual FGenericFrameInfo									OnEncodeDone(const FScalableVideoController::FLayerFrameConfig& Config) override;
-	virtual void												OnRatesUpdated(const FVideoBitrateAllocation& Bitrates) override;
+	UE_API virtual TArray<FScalableVideoController::FLayerFrameConfig> NextFrameConfig(bool bRestart) override;
+	UE_API virtual FGenericFrameInfo									OnEncodeDone(const FScalableVideoController::FLayerFrameConfig& Config) override;
+	UE_API virtual void												OnRatesUpdated(const FVideoBitrateAllocation& Bitrates) override;
 
 private:
 	// NOTE: This enum name is duplicated throughout the ScalabilityStructure variants
@@ -49,15 +51,15 @@ private:
 		ActiveDecodeTargets[Sid * NumTemporalLayers + Tid] = value;
 	}
 
-	bool						   TemporalLayerIsActive(int Tid) const;
-	static EDecodeTargetIndication Dti(int Sid, int Tid, const FScalableVideoController::FLayerFrameConfig& config);
+	UE_API bool						   TemporalLayerIsActive(int Tid) const;
+	static UE_API EDecodeTargetIndication Dti(int Sid, int Tid, const FScalableVideoController::FLayerFrameConfig& config);
 
-	TArray<FScalableVideoController::FLayerFrameConfig> KeyframeConfig();
-	TArray<FScalableVideoController::FLayerFrameConfig> T0Config();
-	TArray<FScalableVideoController::FLayerFrameConfig> T1Config();
-	TArray<FScalableVideoController::FLayerFrameConfig> T2Config(EFramePattern Pattern);
+	UE_API TArray<FScalableVideoController::FLayerFrameConfig> KeyframeConfig();
+	UE_API TArray<FScalableVideoController::FLayerFrameConfig> T0Config();
+	UE_API TArray<FScalableVideoController::FLayerFrameConfig> T1Config();
+	UE_API TArray<FScalableVideoController::FLayerFrameConfig> T2Config(EFramePattern Pattern);
 
-	EFramePattern NextPattern(EFramePattern InLastPattern) const;
+	UE_API EFramePattern NextPattern(EFramePattern InLastPattern) const;
 
 	const int NumSpatialLayers;
 	const int NumTemporalLayers;
@@ -71,7 +73,7 @@ private:
 // S1  0--0--0-
 //     |       ...
 // S0  0--0--0-
-class AVCODECSCORE_API FScalabilityStructureL2T1Key : public FScalabilityStructureKeySvc
+class FScalabilityStructureL2T1Key : public FScalabilityStructureKeySvc
 {
 public:
 	FScalabilityStructureL2T1Key()
@@ -80,7 +82,7 @@ public:
 	}
 	~FScalabilityStructureL2T1Key() override = default;
 
-	FFrameDependencyStructure DependencyStructure() const override;
+	UE_API FFrameDependencyStructure DependencyStructure() const override;
 };
 
 // S1T1     0   0
@@ -91,7 +93,7 @@ public:
 //        |/   /   /
 // S0T0   0---0---0
 // Time-> 0 1 2 3 4
-class AVCODECSCORE_API FScalabilityStructureL2T2Key : public FScalabilityStructureKeySvc
+class FScalabilityStructureL2T2Key : public FScalabilityStructureKeySvc
 {
 public:
 	FScalabilityStructureL2T2Key()
@@ -100,10 +102,10 @@ public:
 	}
 	~FScalabilityStructureL2T2Key() override = default;
 
-	FFrameDependencyStructure DependencyStructure() const override;
+	UE_API FFrameDependencyStructure DependencyStructure() const override;
 };
 
-class AVCODECSCORE_API FScalabilityStructureL2T3Key : public FScalabilityStructureKeySvc
+class FScalabilityStructureL2T3Key : public FScalabilityStructureKeySvc
 {
 public:
 	FScalabilityStructureL2T3Key()
@@ -112,10 +114,10 @@ public:
 	}
 	~FScalabilityStructureL2T3Key() override = default;
 
-	FFrameDependencyStructure DependencyStructure() const override;
+	UE_API FFrameDependencyStructure DependencyStructure() const override;
 };
 
-class AVCODECSCORE_API FScalabilityStructureL3T1Key : public FScalabilityStructureKeySvc
+class FScalabilityStructureL3T1Key : public FScalabilityStructureKeySvc
 {
 public:
 	FScalabilityStructureL3T1Key()
@@ -124,10 +126,10 @@ public:
 	}
 	~FScalabilityStructureL3T1Key() override = default;
 
-	FFrameDependencyStructure DependencyStructure() const override;
+	UE_API FFrameDependencyStructure DependencyStructure() const override;
 };
 
-class AVCODECSCORE_API FScalabilityStructureL3T2Key : public FScalabilityStructureKeySvc
+class FScalabilityStructureL3T2Key : public FScalabilityStructureKeySvc
 {
 public:
 	FScalabilityStructureL3T2Key()
@@ -136,10 +138,10 @@ public:
 	}
 	~FScalabilityStructureL3T2Key() override = default;
 
-	FFrameDependencyStructure DependencyStructure() const override;
+	UE_API FFrameDependencyStructure DependencyStructure() const override;
 };
 
-class AVCODECSCORE_API FScalabilityStructureL3T3Key : public FScalabilityStructureKeySvc
+class FScalabilityStructureL3T3Key : public FScalabilityStructureKeySvc
 {
 public:
 	FScalabilityStructureL3T3Key()
@@ -148,5 +150,7 @@ public:
 	}
 	~FScalabilityStructureL3T3Key() override = default;
 
-	FFrameDependencyStructure DependencyStructure() const override;
-};
\ No newline at end of file
+	UE_API FFrameDependencyStructure DependencyStructure() const override;
+};
+
+#undef UE_API
diff --git a/Engine/Plugins/Experimental/AVCodecs/AVCodecsCore/Source/AVCodecsCore/Public/Video/Encoders/SVC/ScalabilityStructureL2T2KeyShift.h b/Engine/Plugins/Experimental/AVCodecs/AVCodecsCore/Source/AVCodecsCore/Public/Video/Encoders/SVC/ScalabilityStructureL2T2KeyShift.h
--- a/Engine/Plugins/Experimental/AVCodecs/AVCodecsCore/Source/AVCodecsCore/Public/Video/Encoders/SVC/ScalabilityStructureL2T2KeyShift.h
+++ b/Engine/Plugins/Experimental/AVCodecs/AVCodecsCore/Source/AVCodecsCore/Public/Video/Encoders/SVC/ScalabilityStructureL2T2KeyShift.h
@@ -5,6 +5,8 @@
 #include "Video/GenericFrameInfo.h"
 #include "Video/Encoders/SVC/ScalableVideoController.h"
 
+#define UE_API AVCODECSCORE_API
+
 // S1T1     0   0
 //         /   /   /
 // S1T0   0---0---0
@@ -13,18 +15,18 @@
 //        |  /   /
 // S0T0   0-0---0--
 // Time-> 0 1 2 3 4
-class AVCODECSCORE_API FScalabilityStructureL2T2KeyShift : public FScalableVideoController
+class FScalabilityStructureL2T2KeyShift : public FScalableVideoController
 {
 public:
-	FScalabilityStructureL2T2KeyShift();
+	UE_API FScalabilityStructureL2T2KeyShift();
 	~FScalabilityStructureL2T2KeyShift() override = default;
 
-	virtual FStreamLayersConfig		  StreamConfig() const override;
-	virtual FFrameDependencyStructure DependencyStructure() const override;
+	UE_API virtual FStreamLayersConfig		  StreamConfig() const override;
+	UE_API virtual FFrameDependencyStructure DependencyStructure() const override;
 
-	virtual TArray<FScalableVideoController::FLayerFrameConfig> NextFrameConfig(bool bRestart) override;
-	virtual FGenericFrameInfo									OnEncodeDone(const FScalableVideoController::FLayerFrameConfig& Config) override;
-	virtual void												OnRatesUpdated(const FVideoBitrateAllocation& Bitrates) override;
+	UE_API virtual TArray<FScalableVideoController::FLayerFrameConfig> NextFrameConfig(bool bRestart) override;
+	UE_API virtual FGenericFrameInfo									OnEncodeDone(const FScalableVideoController::FLayerFrameConfig& Config) override;
+	UE_API virtual void												OnRatesUpdated(const FVideoBitrateAllocation& Bitrates) override;
 
 private:
 	// NOTE: This enum name is duplicated throughout the ScalabilityStructure variants
@@ -47,8 +49,10 @@ private:
 	{
 		ActiveDecodeTargets[Sid * NumTemporalLayers + Tid] = value;
 	}
-	static EDecodeTargetIndication Dti(int Sid, int Tid, const FScalableVideoController::FLayerFrameConfig& Config);
+	static UE_API EDecodeTargetIndication Dti(int Sid, int Tid, const FScalableVideoController::FLayerFrameConfig& Config);
 
 	EFramePattern NextPattern = EFramePattern::Key;
 	TArray<bool>  ActiveDecodeTargets;
-};
\ No newline at end of file
+};
+
+#undef UE_API
diff --git a/Engine/Plugins/Experimental/AVCodecs/AVCodecsCore/Source/AVCodecsCore/Public/Video/Encoders/SVC/ScalabilityStructureSimulcast.h b/Engine/Plugins/Experimental/AVCodecs/AVCodecsCore/Source/AVCodecsCore/Public/Video/Encoders/SVC/ScalabilityStructureSimulcast.h
--- a/Engine/Plugins/Experimental/AVCodecs/AVCodecsCore/Source/AVCodecsCore/Public/Video/Encoders/SVC/ScalabilityStructureSimulcast.h
+++ b/Engine/Plugins/Experimental/AVCodecs/AVCodecsCore/Source/AVCodecsCore/Public/Video/Encoders/SVC/ScalabilityStructureSimulcast.h
@@ -5,18 +5,20 @@
 #include "Video/GenericFrameInfo.h"
 #include "Video/Encoders/SVC/ScalableVideoController.h"
 
+#define UE_API AVCODECSCORE_API
+
 // Scalability structure with multiple independent spatial layers each with the
 // same temporal layering.
-class AVCODECSCORE_API FScalabilityStructureSimulcast : public FScalableVideoController
+class FScalabilityStructureSimulcast : public FScalableVideoController
 {
 public:
-	FScalabilityStructureSimulcast(int NumSpatialLayers, int NumTemporalLayers, FIntFraction ResolutionFactor);
+	UE_API FScalabilityStructureSimulcast(int NumSpatialLayers, int NumTemporalLayers, FIntFraction ResolutionFactor);
 	virtual ~FScalabilityStructureSimulcast() override = default;
 
-	virtual FStreamLayersConfig									StreamConfig() const override;
-	virtual TArray<FScalableVideoController::FLayerFrameConfig> NextFrameConfig(bool bRestart) override;
-	virtual FGenericFrameInfo									OnEncodeDone(const FScalableVideoController::FLayerFrameConfig& Config) override;
-	virtual void												OnRatesUpdated(const FVideoBitrateAllocation& Bitrates) override;
+	UE_API virtual FStreamLayersConfig									StreamConfig() const override;
+	UE_API virtual TArray<FScalableVideoController::FLayerFrameConfig> NextFrameConfig(bool bRestart) override;
+	UE_API virtual FGenericFrameInfo									OnEncodeDone(const FScalableVideoController::FLayerFrameConfig& Config) override;
+	UE_API virtual void												OnRatesUpdated(const FVideoBitrateAllocation& Bitrates) override;
 
 private:
 	// NOTE: This enum name is duplicated throughout the ScalabilityStructure variants
@@ -48,9 +50,9 @@ private:
 		ActiveDecodeTargets[Sid * NumTemporalLayers + Tid] = Value;
 	}
 
-	EFramePattern				   NextPattern() const;
-	bool						   TemporalLayerIsActive(int Tid) const;
-	static EDecodeTargetIndication Dti(int Sid, int Tid, const FScalableVideoController::FLayerFrameConfig& Config);
+	UE_API EFramePattern				   NextPattern() const;
+	UE_API bool						   TemporalLayerIsActive(int Tid) const;
+	static UE_API EDecodeTargetIndication Dti(int Sid, int Tid, const FScalableVideoController::FLayerFrameConfig& Config);
 
 	const int	   NumSpatialLayers;
 	const int	   NumTemporalLayers;
@@ -65,7 +67,7 @@ private:
 // S1  0--0--0-
 //             ...
 // S0  0--0--0-
-class AVCODECSCORE_API FScalabilityStructureS2T1 : public FScalabilityStructureSimulcast
+class FScalabilityStructureS2T1 : public FScalabilityStructureSimulcast
 {
 public:
 	FScalabilityStructureS2T1(FIntFraction ResolutionFactor = { 1, 2 })
@@ -74,17 +76,17 @@ public:
 	}
 	~FScalabilityStructureS2T1() override = default;
 
-	FFrameDependencyStructure DependencyStructure() const override;
+	UE_API FFrameDependencyStructure DependencyStructure() const override;
 };
 
-class AVCODECSCORE_API FScalabilityStructureS2T2 : public FScalabilityStructureSimulcast
+class FScalabilityStructureS2T2 : public FScalabilityStructureSimulcast
 {
 public:
 	explicit FScalabilityStructureS2T2(FIntFraction ResolutionFactor = { 1, 2 })
 		: FScalabilityStructureSimulcast(2, 2, ResolutionFactor) {}
 	~FScalabilityStructureS2T2() override = default;
 
-	FFrameDependencyStructure DependencyStructure() const override;
+	UE_API FFrameDependencyStructure DependencyStructure() const override;
 };
 
 // S1T2       3   7
@@ -99,7 +101,7 @@ public:
 //           |_/
 // S0T0     0-------8...
 // Time->   0 1 2 3 4
-class AVCODECSCORE_API FScalabilityStructureS2T3 : public FScalabilityStructureSimulcast
+class FScalabilityStructureS2T3 : public FScalabilityStructureSimulcast
 {
 public:
 	explicit FScalabilityStructureS2T3(FIntFraction ResolutionFactor = { 1, 2 })
@@ -108,20 +110,20 @@ public:
 	}
 	~FScalabilityStructureS2T3() override = default;
 
-	FFrameDependencyStructure DependencyStructure() const override;
+	UE_API FFrameDependencyStructure DependencyStructure() const override;
 };
 
-class AVCODECSCORE_API FScalabilityStructureS3T1 : public FScalabilityStructureSimulcast
+class FScalabilityStructureS3T1 : public FScalabilityStructureSimulcast
 {
 public:
 	explicit FScalabilityStructureS3T1(FIntFraction ResolutionFactor = { 1, 2 })
 		: FScalabilityStructureSimulcast(3, 1, ResolutionFactor) {}
 	~FScalabilityStructureS3T1() override = default;
 
-	FFrameDependencyStructure DependencyStructure() const override;
+	UE_API FFrameDependencyStructure DependencyStructure() const override;
 };
 
-class AVCODECSCORE_API FScalabilityStructureS3T2 : public FScalabilityStructureSimulcast
+class FScalabilityStructureS3T2 : public FScalabilityStructureSimulcast
 {
 public:
 	explicit FScalabilityStructureS3T2(FIntFraction ResolutionFactor = { 1, 2 })
@@ -130,10 +132,10 @@ public:
 	}
 	~FScalabilityStructureS3T2() override = default;
 
-	FFrameDependencyStructure DependencyStructure() const override;
+	UE_API FFrameDependencyStructure DependencyStructure() const override;
 };
 
-class AVCODECSCORE_API FScalabilityStructureS3T3 : public FScalabilityStructureSimulcast
+class FScalabilityStructureS3T3 : public FScalabilityStructureSimulcast
 {
 public:
 	FScalabilityStructureS3T3(FIntFraction ResolutionFactor = { 1, 2 })
@@ -142,5 +144,7 @@ public:
 	}
 	~FScalabilityStructureS3T3() override = default;
 
-	FFrameDependencyStructure DependencyStructure() const override;
-};
\ No newline at end of file
+	UE_API FFrameDependencyStructure DependencyStructure() const override;
+};
+
+#undef UE_API
diff --git a/Engine/Plugins/Experimental/AVCodecs/AVCodecsCore/Source/AVCodecsCore/Public/Video/Encoders/SVC/ScalableVideoController.h b/Engine/Plugins/Experimental/AVCodecs/AVCodecsCore/Source/AVCodecsCore/Public/Video/Encoders/SVC/ScalableVideoController.h
--- a/Engine/Plugins/Experimental/AVCodecs/AVCodecsCore/Source/AVCodecsCore/Public/Video/Encoders/SVC/ScalableVideoController.h
+++ b/Engine/Plugins/Experimental/AVCodecs/AVCodecsCore/Source/AVCodecsCore/Public/Video/Encoders/SVC/ScalableVideoController.h
@@ -5,10 +5,12 @@
 #include "Video/GenericFrameInfo.h"
 #include "Video/Encoders/VideoBitrateAllocator.h"
 
+#define UE_API AVCODECSCORE_API
+
 // Controls how video should be encoded to be scalable. Outputs results as
 // buffer usage configuration for encoder and enough details to communicate the
 // scalability structure via dependency descriptor rtp header extension.
-class AVCODECSCORE_API FScalableVideoController
+class FScalableVideoController
 {
 public:
 	struct FIntFraction
@@ -28,24 +30,24 @@ public:
 		FIntFraction ScalingFactors[4] = { { 1, 1 }, { 1, 1 }, { 1, 1 }, { 1, 1 } };
 	};
 
-	class AVCODECSCORE_API FLayerFrameConfig
+	class FLayerFrameConfig
 	{
 	public:
 		// Builders/setters.
-		FLayerFrameConfig& Id(int Value);
-		FLayerFrameConfig& Keyframe();
-		FLayerFrameConfig& SpatialLayerId(int Value);
-		FLayerFrameConfig& TemporalLayerId(int Value);
-		FLayerFrameConfig& Reference(int BufferId);
-		FLayerFrameConfig& Update(int BufferId);
-		FLayerFrameConfig& ReferenceAndUpdate(int BufferId);
+		UE_API FLayerFrameConfig& Id(int Value);
+		UE_API FLayerFrameConfig& Keyframe();
+		UE_API FLayerFrameConfig& SpatialLayerId(int Value);
+		UE_API FLayerFrameConfig& TemporalLayerId(int Value);
+		UE_API FLayerFrameConfig& Reference(int BufferId);
+		UE_API FLayerFrameConfig& Update(int BufferId);
+		UE_API FLayerFrameConfig& ReferenceAndUpdate(int BufferId);
 
 		// Getters.
-		int								 GetId() const;
-		bool							 GetIsKeyframe() const;
-		int								 GetSpatialId() const;
-		int								 GetTemporalId() const;
-		const TArray<FCodecBufferUsage>& GetBuffers() const;
+		UE_API int								 GetId() const;
+		UE_API bool							 GetIsKeyframe() const;
+		UE_API int								 GetSpatialId() const;
+		UE_API int								 GetTemporalId() const;
+		UE_API const TArray<FCodecBufferUsage>& GetBuffers() const;
 
 	private:
 		// Id to match configuration returned by NextFrameConfig with
@@ -87,4 +89,6 @@ public:
 
 	// Returns configuration to pass to EncoderCallback.
 	virtual FGenericFrameInfo OnEncodeDone(const FLayerFrameConfig& Config) = 0;
-};
\ No newline at end of file
+};
+
+#undef UE_API
diff --git a/Engine/Plugins/Experimental/AVCodecs/AVCodecsCore/Source/AVCodecsCore/Public/Video/Encoders/SVC/ScalableVideoControllerNoLayering.h b/Engine/Plugins/Experimental/AVCodecs/AVCodecsCore/Source/AVCodecsCore/Public/Video/Encoders/SVC/ScalableVideoControllerNoLayering.h
--- a/Engine/Plugins/Experimental/AVCodecs/AVCodecsCore/Source/AVCodecsCore/Public/Video/Encoders/SVC/ScalableVideoControllerNoLayering.h
+++ b/Engine/Plugins/Experimental/AVCodecs/AVCodecsCore/Source/AVCodecsCore/Public/Video/Encoders/SVC/ScalableVideoControllerNoLayering.h
@@ -4,19 +4,23 @@
 
 #include "Video/Encoders/SVC/ScalableVideoController.h"
 
-class AVCODECSCORE_API FScalableVideoControllerNoLayering : public FScalableVideoController
+#define UE_API AVCODECSCORE_API
+
+class FScalableVideoControllerNoLayering : public FScalableVideoController
 {
 public:
 	~FScalableVideoControllerNoLayering() override = default;
 
-	virtual FStreamLayersConfig		  StreamConfig() const override;
-	virtual FFrameDependencyStructure DependencyStructure() const override;
+	UE_API virtual FStreamLayersConfig		  StreamConfig() const override;
+	UE_API virtual FFrameDependencyStructure DependencyStructure() const override;
 
-	virtual TArray<FScalableVideoController::FLayerFrameConfig> NextFrameConfig(bool bRestart) override;
-	virtual FGenericFrameInfo									OnEncodeDone(const FScalableVideoController::FLayerFrameConfig& Config) override;
-	virtual void												OnRatesUpdated(const FVideoBitrateAllocation& Bitrates) override;
+	UE_API virtual TArray<FScalableVideoController::FLayerFrameConfig> NextFrameConfig(bool bRestart) override;
+	UE_API virtual FGenericFrameInfo									OnEncodeDone(const FScalableVideoController::FLayerFrameConfig& Config) override;
+	UE_API virtual void												OnRatesUpdated(const FVideoBitrateAllocation& Bitrates) override;
 
 private:
 	bool bStart = true;
 	bool bEnabled = true;
-};
\ No newline at end of file
+};
+
+#undef UE_API
diff --git a/Engine/Plugins/Experimental/AVCodecs/AVCodecsCore/Source/AVCodecsCore/Public/Video/Encoders/SVC/VideoBitrateAllocatorSVC.h b/Engine/Plugins/Experimental/AVCodecs/AVCodecsCore/Source/AVCodecsCore/Public/Video/Encoders/SVC/VideoBitrateAllocatorSVC.h
--- a/Engine/Plugins/Experimental/AVCodecs/AVCodecsCore/Source/AVCodecsCore/Public/Video/Encoders/SVC/VideoBitrateAllocatorSVC.h
+++ b/Engine/Plugins/Experimental/AVCodecs/AVCodecsCore/Source/AVCodecsCore/Public/Video/Encoders/SVC/VideoBitrateAllocatorSVC.h
@@ -5,18 +5,20 @@
 #include "Video/Encoders/VideoBitrateAllocator.h"
 #include "Video/VideoEncoder.h"
 
-class AVCODECSCORE_API FVideoBitrateAllocatorSVC : public FVideoBitrateAllocator
+#define UE_API AVCODECSCORE_API
+
+class FVideoBitrateAllocatorSVC : public FVideoBitrateAllocator
 {
 public:
 	// NOTE: Config stores its rates as kbps whereas this class will use bps. Convert to bps everywhere
-	explicit FVideoBitrateAllocatorSVC(const FVideoEncoderConfig& Config);
+	UE_API explicit FVideoBitrateAllocatorSVC(const FVideoEncoderConfig& Config);
 	virtual ~FVideoBitrateAllocatorSVC() = default;
 
-	FVideoBitrateAllocation Allocate(FVideoBitrateAllocationParameters Parameters) override;
+	UE_API FVideoBitrateAllocation Allocate(FVideoBitrateAllocationParameters Parameters) override;
 
-	static uint32		  GetMaxBitrate(const FVideoEncoderConfig& Config);
-	static uint32		  GetPaddingBitrate(const FVideoEncoderConfig& Config);
-	static TArray<uint32> GetLayerStartBitrates(const FVideoEncoderConfig& Config);
+	static UE_API uint32		  GetMaxBitrate(const FVideoEncoderConfig& Config);
+	static UE_API uint32		  GetPaddingBitrate(const FVideoEncoderConfig& Config);
+	static UE_API TArray<uint32> GetLayerStartBitrates(const FVideoEncoderConfig& Config);
 
 private:
 	struct FNumLayers
@@ -25,15 +27,17 @@ private:
 		size_t Temporal = 1;
 	};
 
-	static FNumLayers		GetNumLayers(const FVideoEncoderConfig& Config);
-	FVideoBitrateAllocation GetAllocationNormalVideo(uint32 TotalBitrateBps, size_t FirstActiveLayer, size_t NumSpatialLayers) const;
+	static UE_API FNumLayers		GetNumLayers(const FVideoEncoderConfig& Config);
+	UE_API FVideoBitrateAllocation GetAllocationNormalVideo(uint32 TotalBitrateBps, size_t FirstActiveLayer, size_t NumSpatialLayers) const;
 
 	// Returns the number of layers that are active and have enough bitrate to
 	// actually be enabled.
-	size_t FindNumEnabledLayers(uint32 TargetRate) const;
+	UE_API size_t FindNumEnabledLayers(uint32 TargetRate) const;
 
 	FVideoEncoderConfig	 Config;
 	const FNumLayers	 NumLayers;
 	const TArray<uint32> CumulativeLayerStartBitrates;
 	size_t				 LastActiveLayerCount;
-};
\ No newline at end of file
+};
+
+#undef UE_API
diff --git a/Engine/Plugins/Experimental/AVCodecs/AVCodecsCore/Source/AVCodecsCore/Public/Video/Encoders/VideoBitrateAllocation.h b/Engine/Plugins/Experimental/AVCodecs/AVCodecsCore/Source/AVCodecsCore/Public/Video/Encoders/VideoBitrateAllocation.h
--- a/Engine/Plugins/Experimental/AVCodecs/AVCodecsCore/Source/AVCodecsCore/Public/Video/Encoders/VideoBitrateAllocation.h
+++ b/Engine/Plugins/Experimental/AVCodecs/AVCodecsCore/Source/AVCodecsCore/Public/Video/Encoders/VideoBitrateAllocation.h
@@ -7,41 +7,43 @@
 #include "Misc/Optional.h"
 #include "Video/VideoEncoder.h"
 
-class AVCODECSCORE_API FVideoBitrateAllocation
+#define UE_API AVCODECSCORE_API
+
+class FVideoBitrateAllocation
 {
 public:
 	static constexpr uint32 MaxBitrateBps = TNumericLimits<uint32>::Max();
 
-	FVideoBitrateAllocation();
+	UE_API FVideoBitrateAllocation();
 
 	virtual ~FVideoBitrateAllocation() = default;
 
-	bool HasBitrate(uint64 SpatialIndex, uint64 TemporalIndex) const;
+	UE_API bool HasBitrate(uint64 SpatialIndex, uint64 TemporalIndex) const;
 
-	uint32 GetBitrate(uint64 SpatialIndex, uint64 TemporalIndex) const;
+	UE_API uint32 GetBitrate(uint64 SpatialIndex, uint64 TemporalIndex) const;
 
-	bool IsSpatialLayerUsed(uint64 SpatialIndex) const;
+	UE_API bool IsSpatialLayerUsed(uint64 SpatialIndex) const;
 
-	uint32 GetSpatialLayerSumBitrate(uint64 SpatialIndex) const;
+	UE_API uint32 GetSpatialLayerSumBitrate(uint64 SpatialIndex) const;
 
-	uint32 GetTemporalLayerSumBitrate(uint64 SpatialIndex, uint64 TemporalIndex) const;
+	UE_API uint32 GetTemporalLayerSumBitrate(uint64 SpatialIndex, uint64 TemporalIndex) const;
 
-	TArray<uint32> GetTemporalLayerAllocation(uint64 SpatialIndex) const;
+	UE_API TArray<uint32> GetTemporalLayerAllocation(uint64 SpatialIndex) const;
 
-	uint32 GetSumBps() const;
+	UE_API uint32 GetSumBps() const;
 
-	bool IsBwLimited() const;
+	UE_API bool IsBwLimited() const;
 
-	bool SetBitrate(size_t SpatialIndex, size_t TemporalIndex, uint32 BitrateBps);
+	UE_API bool SetBitrate(size_t SpatialIndex, size_t TemporalIndex, uint32 BitrateBps);
 
 	// Indicates if the allocation has some layers/streams disabled due to
 	// low available bandwidth.
-	void SetBwLimited(bool Limited);
+	UE_API void SetBwLimited(bool Limited);
 
 	// Returns one FVideoBitrateAllocation for each spatial layer. This is used to
 	// configure simulcast streams. Note that the length of the returned vector is
 	// always 5, the optional is unset for unused layers.
-	TArray<TOptional<FVideoBitrateAllocation>> GetSimulcastAllocations() const;
+	UE_API TArray<TOptional<FVideoBitrateAllocation>> GetSimulcastAllocations() const;
 
 	inline bool operator==(const FVideoBitrateAllocation& Other) const
 	{
@@ -67,4 +69,6 @@ private:
 	uint32			  SumBps;
 	TOptional<uint32> Bitrates[Video::MaxSpatialLayers][Video::MaxTemporalStreams];
 	bool			  bIsBwLimited;
-};
\ No newline at end of file
+};
+
+#undef UE_API
diff --git a/Engine/Plugins/Experimental/AVCodecs/AVCodecsCore/Source/AVCodecsCore/Public/Video/Encoders/VideoBitrateAllocator.h b/Engine/Plugins/Experimental/AVCodecs/AVCodecsCore/Source/AVCodecsCore/Public/Video/Encoders/VideoBitrateAllocator.h
--- a/Engine/Plugins/Experimental/AVCodecs/AVCodecsCore/Source/AVCodecsCore/Public/Video/Encoders/VideoBitrateAllocator.h
+++ b/Engine/Plugins/Experimental/AVCodecs/AVCodecsCore/Source/AVCodecsCore/Public/Video/Encoders/VideoBitrateAllocator.h
@@ -5,7 +5,9 @@
 #include "Misc/FrameRate.h"
 #include "Video/Encoders/VideoBitrateAllocation.h"
 
-class AVCODECSCORE_API FVideoBitrateAllocationParameters
+#define UE_API AVCODECSCORE_API
+
+class FVideoBitrateAllocationParameters
 {
 public:
 	FVideoBitrateAllocationParameters(uint32 TotalBitrateBps, FFrameRate Framerate)
@@ -26,13 +28,15 @@ public:
 	FFrameRate Framerate;
 };
 
-class AVCODECSCORE_API FVideoBitrateAllocator
+class FVideoBitrateAllocator
 {
 public:
 	FVideoBitrateAllocator() = default;
 	~FVideoBitrateAllocator() = default;
 
-	virtual FVideoBitrateAllocation GetAllocation(uint32 TotalBitrateBps, FFrameRate Framerate);
+	UE_API virtual FVideoBitrateAllocation GetAllocation(uint32 TotalBitrateBps, FFrameRate Framerate);
+
+	UE_API virtual FVideoBitrateAllocation Allocate(FVideoBitrateAllocationParameters Parameters);
+};
 
-	virtual FVideoBitrateAllocation Allocate(FVideoBitrateAllocationParameters Parameters);
-};
\ No newline at end of file
+#undef UE_API
diff --git a/Engine/Plugins/Experimental/AVCodecs/AVCodecsCore/Source/AVCodecsCore/Public/Video/GenericFrameInfo.h b/Engine/Plugins/Experimental/AVCodecs/AVCodecsCore/Source/AVCodecsCore/Public/Video/GenericFrameInfo.h
--- a/Engine/Plugins/Experimental/AVCodecs/AVCodecsCore/Source/AVCodecsCore/Public/Video/GenericFrameInfo.h
+++ b/Engine/Plugins/Experimental/AVCodecs/AVCodecsCore/Source/AVCodecsCore/Public/Video/GenericFrameInfo.h
@@ -5,17 +5,19 @@
 #include "Containers/Array.h"
 #include "Video/DependencyDescriptor.h"
 
-struct AVCODECSCORE_API FCodecBufferUsage
+#define UE_API AVCODECSCORE_API
+
+struct FCodecBufferUsage
 {
 	int32 Id = 0;
 	bool  bReferenced = false;
 	bool  bUpdated = false;
 };
 
-class AVCODECSCORE_API FGenericFrameInfo : public FFrameDependencyTemplate
+class FGenericFrameInfo : public FFrameDependencyTemplate
 {
 public:
-	FGenericFrameInfo();
+	UE_API FGenericFrameInfo();
 
 	friend bool operator==(const FGenericFrameInfo& Lhs, const FGenericFrameInfo& Rhs)
 	{
@@ -29,4 +31,6 @@ public:
 	TArray<FCodecBufferUsage> EncoderBuffers;
 	TArray<bool>			  PartOfChain;
 	TArray<bool>			  ActiveDecodeTargets;
-};
\ No newline at end of file
+};
+
+#undef UE_API
diff --git a/Engine/Plugins/Experimental/AVCodecs/AVCodecsCore/Source/AVCodecsCore/Public/Video/Resources/D3D/VideoResourceD3D.h b/Engine/Plugins/Experimental/AVCodecs/AVCodecsCore/Source/AVCodecsCore/Public/Video/Resources/D3D/VideoResourceD3D.h
--- a/Engine/Plugins/Experimental/AVCodecs/AVCodecsCore/Source/AVCodecsCore/Public/Video/Resources/D3D/VideoResourceD3D.h
+++ b/Engine/Plugins/Experimental/AVCodecs/AVCodecsCore/Source/AVCodecsCore/Public/Video/Resources/D3D/VideoResourceD3D.h
@@ -35,40 +35,40 @@ THIRD_PARTY_INCLUDES_END
  * D3D11 and D3D12 platform video context and resource.
  */
 
-class AVCODECSCORE_API FVideoContextD3D11 : public FAVContext
+class FVideoContextD3D11 : public FAVContext
 {
 public:
 	TRefCountPtr<ID3D11Device> Device;
 
-	FVideoContextD3D11(TRefCountPtr<ID3D11Device> const& Device);
+	AVCODECSCORE_API FVideoContextD3D11(TRefCountPtr<ID3D11Device> const& Device);
 };
 
-class AVCODECSCORE_API FVideoContextD3D12 : public FAVContext
+class FVideoContextD3D12 : public FAVContext
 {
 public:
 	TRefCountPtr<ID3D12Device> Device;
 
-	FVideoContextD3D12(TRefCountPtr<ID3D12Device> const& Device);
+	AVCODECSCORE_API FVideoContextD3D12(TRefCountPtr<ID3D12Device> const& Device);
 };
 
-class AVCODECSCORE_API FVideoResourceD3D11 : public TVideoResource<FVideoContextD3D11>
+class FVideoResourceD3D11 : public TVideoResource<FVideoContextD3D11>
 {
 private:
 	TRefCountPtr<ID3D11Texture2D> Raw;
 	HANDLE SharedHandle = nullptr;
 
 public:
-	static FVideoDescriptor GetDescriptorFrom(TSharedRef<FAVDevice> const& Device, TRefCountPtr<ID3D11Texture2D> const& Raw);
+	static AVCODECSCORE_API FVideoDescriptor GetDescriptorFrom(TSharedRef<FAVDevice> const& Device, TRefCountPtr<ID3D11Texture2D> const& Raw);
 
 	inline TRefCountPtr<ID3D11Texture2D> const& GetRaw() const { return Raw; }
 	inline HANDLE const& GetSharedHandle() const { return SharedHandle; }
 
-	FVideoResourceD3D11(TSharedRef<FAVDevice> const& Device, TRefCountPtr<ID3D11Texture2D> const& Raw, FAVLayout const& Layout);
+	AVCODECSCORE_API FVideoResourceD3D11(TSharedRef<FAVDevice> const& Device, TRefCountPtr<ID3D11Texture2D> const& Raw, FAVLayout const& Layout);
 
-	virtual FAVResult Validate() const override;
+	AVCODECSCORE_API virtual FAVResult Validate() const override;
 };
 
-class AVCODECSCORE_API FVideoResourceD3D12 : public TVideoResource<FVideoContextD3D12>
+class FVideoResourceD3D12 : public TVideoResource<FVideoContextD3D12>
 {
 public:
 	struct FRawD3D12
@@ -86,7 +86,7 @@ private:
 	FRawD3D12 Raw;
 
 public:
-	static FVideoDescriptor GetDescriptorFrom(TSharedRef<FAVDevice> const& Device, FRawD3D12 const& Raw);
+	static AVCODECSCORE_API FVideoDescriptor GetDescriptorFrom(TSharedRef<FAVDevice> const& Device, FRawD3D12 const& Raw);
 
 	inline TRefCountPtr<ID3D12Resource> const& GetResource() const { return Raw.D3DResource; }
 	inline HANDLE const& GetResourceSharedHandle() const { return Raw.D3DResourceShared; }
@@ -98,13 +98,13 @@ public:
 	inline HANDLE const& GetFenceSharedHandle() const { return Raw.D3DFenceShared; }
 	inline uint64 const& GetFenceValue() const { return Raw.FenceValue; }
 
-	uint64 GetSizeInBytes();
+	AVCODECSCORE_API uint64 GetSizeInBytes();
 
-	FVideoResourceD3D12(TSharedRef<FAVDevice> const& Device, FRawD3D12 const& Raw, FAVLayout const& Layout);
-	FVideoResourceD3D12(TSharedRef<FAVDevice> const& Device, FRawD3D12 const& Raw, FAVLayout const& Layout, FVideoDescriptor const& Descriptor);
-	virtual ~FVideoResourceD3D12() override;
+	AVCODECSCORE_API FVideoResourceD3D12(TSharedRef<FAVDevice> const& Device, FRawD3D12 const& Raw, FAVLayout const& Layout);
+	AVCODECSCORE_API FVideoResourceD3D12(TSharedRef<FAVDevice> const& Device, FRawD3D12 const& Raw, FAVLayout const& Layout, FVideoDescriptor const& Descriptor);
+	AVCODECSCORE_API virtual ~FVideoResourceD3D12() override;
 
-	virtual FAVResult Validate() const override;
+	AVCODECSCORE_API virtual FAVResult Validate() const override;
 };
 
 DECLARE_TYPEID(FVideoContextD3D11, AVCODECSCORE_API);
@@ -112,4 +112,4 @@ DECLARE_TYPEID(FVideoContextD3D12, AVCODECSCORE_API);
 DECLARE_TYPEID(FVideoResourceD3D11, AVCODECSCORE_API);
 DECLARE_TYPEID(FVideoResourceD3D12, AVCODECSCORE_API);
 
-#endif
\ No newline at end of file
+#endif
diff --git a/Engine/Plugins/Experimental/AVCodecs/AVCodecsCore/Source/AVCodecsCore/Public/Video/Resources/Metal/VideoResourceMetal.h b/Engine/Plugins/Experimental/AVCodecs/AVCodecsCore/Source/AVCodecsCore/Public/Video/Resources/Metal/VideoResourceMetal.h
--- a/Engine/Plugins/Experimental/AVCodecs/AVCodecsCore/Source/AVCodecsCore/Public/Video/Resources/Metal/VideoResourceMetal.h
+++ b/Engine/Plugins/Experimental/AVCodecs/AVCodecsCore/Source/AVCodecsCore/Public/Video/Resources/Metal/VideoResourceMetal.h
@@ -19,31 +19,31 @@ THIRD_PARTY_INCLUDES_END
  * Metal platform video context and resource.
  */
 
-class AVCODECSCORE_API FVideoContextMetal : public FAVContext
+class FVideoContextMetal : public FAVContext
 {
 public:
 	MTL::Device* Device;
 
-	FVideoContextMetal(MTL::Device* Device);
+	AVCODECSCORE_API FVideoContextMetal(MTL::Device* Device);
 };
 
 
-class AVCODECSCORE_API FVideoResourceMetal : public TVideoResource<FVideoContextMetal>
+class FVideoResourceMetal : public TVideoResource<FVideoContextMetal>
 {
 private:
     MTL::Texture* Raw;
 
 public:
-	static FVideoDescriptor GetDescriptorFrom(TSharedRef<FAVDevice> const& Device, MTL::Texture* Raw);
+	static AVCODECSCORE_API FVideoDescriptor GetDescriptorFrom(TSharedRef<FAVDevice> const& Device, MTL::Texture* Raw);
 
 	inline MTL::Texture* GetRaw() const { return Raw; }
 
-	FVideoResourceMetal(TSharedRef<FAVDevice> const& Device, MTL::Texture* Raw, FAVLayout const& Layout);
-    virtual ~FVideoResourceMetal() override;
+	AVCODECSCORE_API FVideoResourceMetal(TSharedRef<FAVDevice> const& Device, MTL::Texture* Raw, FAVLayout const& Layout);
+    AVCODECSCORE_API virtual ~FVideoResourceMetal() override;
 	
-	FAVResult CopyFrom(CVPixelBufferRef Other);
+	AVCODECSCORE_API FAVResult CopyFrom(CVPixelBufferRef Other);
     
-	virtual FAVResult Validate() const override;
+	AVCODECSCORE_API virtual FAVResult Validate() const override;
 };
 
 DECLARE_TYPEID(FVideoContextMetal, AVCODECSCORE_API);
diff --git a/Engine/Plugins/Experimental/AVCodecs/AVCodecsCore/Source/AVCodecsCore/Public/Video/Resources/VideoResourceCPU.h b/Engine/Plugins/Experimental/AVCodecs/AVCodecsCore/Source/AVCodecsCore/Public/Video/Resources/VideoResourceCPU.h
--- a/Engine/Plugins/Experimental/AVCodecs/AVCodecsCore/Source/AVCodecsCore/Public/Video/Resources/VideoResourceCPU.h
+++ b/Engine/Plugins/Experimental/AVCodecs/AVCodecsCore/Source/AVCodecsCore/Public/Video/Resources/VideoResourceCPU.h
@@ -5,17 +5,19 @@
 #include "AVContext.h"
 #include "Video/VideoResource.h"
 
+#define UE_API AVCODECSCORE_API
+
 /**
  * CPU video context and resource.
  */
 
-class AVCODECSCORE_API FVideoContextCPU : public FAVContext
+class FVideoContextCPU : public FAVContext
 {
 public:
-	FVideoContextCPU();
+	UE_API FVideoContextCPU();
 };
 
-class AVCODECSCORE_API FVideoResourceCPU : public TVideoResource<FVideoContextCPU>
+class FVideoResourceCPU : public TVideoResource<FVideoContextCPU>
 {
 private:
 	TSharedPtr<uint8> Raw;
@@ -24,22 +26,24 @@ public:
 	inline TSharedPtr<uint8> const& GetRaw() const { return Raw; }
 	inline void					 SetRaw(TSharedPtr<uint8> RawData) { Raw = RawData; }
 
-	FVideoResourceCPU(TSharedRef<FAVDevice> const& Device, TSharedPtr<uint8> const& Raw, FAVLayout const& Layout, FVideoDescriptor const& Descriptor);
+	UE_API FVideoResourceCPU(TSharedRef<FAVDevice> const& Device, TSharedPtr<uint8> const& Raw, FAVLayout const& Layout, FVideoDescriptor const& Descriptor);
 	virtual ~FVideoResourceCPU() override = default;
 
-	virtual FAVResult Validate() const override;
+	UE_API virtual FAVResult Validate() const override;
 };
 
-class AVCODECSCORE_API FResolvableVideoResourceCPU : public TResolvableVideoResource<FVideoResourceCPU>
+class FResolvableVideoResourceCPU : public TResolvableVideoResource<FVideoResourceCPU>
 {
 protected:
-	virtual TSharedPtr<FVideoResourceCPU> TryResolve(TSharedPtr<FAVDevice> const& Device, FVideoDescriptor const& Descriptor) override;
+	UE_API virtual TSharedPtr<FVideoResourceCPU> TryResolve(TSharedPtr<FAVDevice> const& Device, FVideoDescriptor const& Descriptor) override;
 
 private:
-	uint32 GetStrideInterleavedOrLuma(FVideoDescriptor const& Descriptor);
-	uint32 GetSize(FVideoDescriptor const& Descriptor);
+	UE_API uint32 GetStrideInterleavedOrLuma(FVideoDescriptor const& Descriptor);
+	UE_API uint32 GetSize(FVideoDescriptor const& Descriptor);
 };
 
 DECLARE_TYPEID(FVideoContextCPU, AVCODECSCORE_API);
 DECLARE_TYPEID(FVideoResourceCPU, AVCODECSCORE_API);
-DECLARE_TYPEID(FResolvableVideoResourceCPU, AVCODECSCORE_API);
\ No newline at end of file
+DECLARE_TYPEID(FResolvableVideoResourceCPU, AVCODECSCORE_API);
+
+#undef UE_API
diff --git a/Engine/Plugins/Experimental/AVCodecs/AVCodecsCore/Source/AVCodecsCore/Public/Video/Resources/Vulkan/VideoResourceVulkan.h b/Engine/Plugins/Experimental/AVCodecs/AVCodecsCore/Source/AVCodecsCore/Public/Video/Resources/Vulkan/VideoResourceVulkan.h
--- a/Engine/Plugins/Experimental/AVCodecs/AVCodecsCore/Source/AVCodecsCore/Public/Video/Resources/Vulkan/VideoResourceVulkan.h
+++ b/Engine/Plugins/Experimental/AVCodecs/AVCodecsCore/Source/AVCodecsCore/Public/Video/Resources/Vulkan/VideoResourceVulkan.h
@@ -21,7 +21,7 @@ THIRD_PARTY_INCLUDES_END
  * Vulkan platform video context and resource.
  */
 
-class AVCODECSCORE_API FVideoContextVulkan : public FAVContext
+class FVideoContextVulkan : public FAVContext
 {
 public:
 	VkInstance Instance;
@@ -29,7 +29,7 @@ public:
 	VkPhysicalDevice PhysicalDevice;
 	TFunction<PFN_vkVoidFunction(const char*)> vkGetDeviceProcAddr;
 
-	FVideoContextVulkan(VkInstance const& Instance, VkDevice const& Device, VkPhysicalDevice const& PhysicalDevice, TFunction<PFN_vkVoidFunction(const char*)> const& GetDeviceProcAddrFunc);
+	AVCODECSCORE_API FVideoContextVulkan(VkInstance const& Instance, VkDevice const& Device, VkPhysicalDevice const& PhysicalDevice, TFunction<PFN_vkVoidFunction(const char*)> const& GetDeviceProcAddrFunc);
 };
 
 #if PLATFORM_WINDOWS
@@ -38,7 +38,7 @@ typedef HANDLE FVulkanSharedHandle;
 typedef int FVulkanSharedHandle;
 #endif
 
-class AVCODECSCORE_API FVideoResourceVulkan : public TVideoResource<FVideoContextVulkan>
+class FVideoResourceVulkan : public TVideoResource<FVideoContextVulkan>
 {
 private:
 	VkDeviceMemory DeviceMemory;
@@ -51,10 +51,10 @@ public:
 	inline FVulkanSharedHandle const& GetSharedHandle() const { return DeviceMemorySharedHandle; }
 	inline FVulkanSharedHandle const& GetFenceSharedHandle() const { return FenceSharedHandle; }
 
-	FVideoResourceVulkan(TSharedRef<FAVDevice> const& Device, VkDeviceMemory Raw, FAVLayout const& Layout, FVideoDescriptor const& Descriptor);
+	AVCODECSCORE_API FVideoResourceVulkan(TSharedRef<FAVDevice> const& Device, VkDeviceMemory Raw, FAVLayout const& Layout, FVideoDescriptor const& Descriptor);
 	virtual ~FVideoResourceVulkan() override = default;
 
-	virtual FAVResult Validate() const override;
+	AVCODECSCORE_API virtual FAVResult Validate() const override;
 };
 
 DECLARE_TYPEID(FVideoContextVulkan, AVCODECSCORE_API);
diff --git a/Engine/Plugins/Experimental/AVCodecs/AVCodecsCore/Source/AVCodecsCore/Public/Video/VideoCodecSpecificInfo.h b/Engine/Plugins/Experimental/AVCodecs/AVCodecsCore/Source/AVCodecsCore/Public/Video/VideoCodecSpecificInfo.h
--- a/Engine/Plugins/Experimental/AVCodecs/AVCodecsCore/Source/AVCodecsCore/Public/Video/VideoCodecSpecificInfo.h
+++ b/Engine/Plugins/Experimental/AVCodecs/AVCodecsCore/Source/AVCodecsCore/Public/Video/VideoCodecSpecificInfo.h
@@ -8,7 +8,7 @@
 #include "Video/GenericFrameInfo.h"
 #include "Video/VideoConfig.h"
 
-struct AVCODECSCORE_API FCodecSpecificInfoVP8
+struct FCodecSpecificInfoVP8
 {
 	bool  bNonReference;
 	uint8 TemporalIdx;
@@ -31,7 +31,7 @@ struct AVCODECSCORE_API FCodecSpecificInfoVP8
 	size_t					UpdatedBuffersCount;
 };
 
-struct AVCODECSCORE_API FCodecSpecificInfoVP9
+struct FCodecSpecificInfoVP9
 {
 
 	bool bFirstFrameInPicture; // First frame, increment picture_id.
@@ -60,7 +60,7 @@ struct AVCODECSCORE_API FCodecSpecificInfoVP9
 	uint8 PDiff[UE::AVCodecCore::VP9::MaxRefPics];
 };
 
-struct AVCODECSCORE_API FCodecSpecificInfoH264
+struct FCodecSpecificInfoH264
 {
 	EH264PacketizationMode PacketizationMode;
 	uint8				   TemporalIdx;
@@ -75,7 +75,7 @@ union AVCODECSCORE_API FCodecSpecificInfoUnion
 	FCodecSpecificInfoH264 H264;
 };
 
-struct AVCODECSCORE_API FCodecSpecificInfo
+struct FCodecSpecificInfo
 {
 	FCodecSpecificInfo()
 		: Codec(EVideoCodec::Undefined)
@@ -90,4 +90,4 @@ struct AVCODECSCORE_API FCodecSpecificInfo
 	bool								 bEndOfPicture = true;
 	TOptional<FGenericFrameInfo>		 GenericFrameInfo;
 	TOptional<FFrameDependencyStructure> TemplateStructure;
-};
\ No newline at end of file
+};
diff --git a/Engine/Plugins/Experimental/AVCodecs/AVCodecsCore/Source/AVCodecsCore/Public/Video/VideoResource.h b/Engine/Plugins/Experimental/AVCodecs/AVCodecsCore/Source/AVCodecsCore/Public/Video/VideoResource.h
--- a/Engine/Plugins/Experimental/AVCodecs/AVCodecsCore/Source/AVCodecsCore/Public/Video/VideoResource.h
+++ b/Engine/Plugins/Experimental/AVCodecs/AVCodecsCore/Source/AVCodecsCore/Public/Video/VideoResource.h
@@ -10,6 +10,8 @@
 
 #include <atomic>
 
+#define UE_API AVCODECSCORE_API
+
 enum class EVideoFormat : uint8
 {
 	BGRA = EPixelFormat::PF_B8G8R8A8,
@@ -161,7 +163,7 @@ public:
 /**
  * Base wrapper for a video device resource.
  */
-class AVCODECSCORE_API FVideoResource : public FAVResource
+class FVideoResource : public FAVResource
 {
 private:
 	/**
@@ -215,7 +217,7 @@ public:
 	 */
 	inline bool IsInUse() { return bIsInUse; }
 
-	FVideoResource(TSharedRef<FAVDevice> const& Device, FAVLayout const& Layout, FVideoDescriptor const& Descriptor);
+	UE_API FVideoResource(TSharedRef<FAVDevice> const& Device, FAVLayout const& Layout, FVideoDescriptor const& Descriptor);
 	virtual ~FVideoResource() override = default;
 };
 
@@ -274,3 +276,5 @@ protected:
 		return nullptr;
 	}
 };
+
+#undef UE_API
diff --git a/Engine/Plugins/Experimental/AVCodecs/AVCodecsCore/Source/AVCodecsCoreRHI/Private/Audio/Encoders/SimpleAudioEncoder.cpp b/Engine/Plugins/Experimental/AVCodecs/AVCodecsCore/Source/AVCodecsCoreRHI/Private/Audio/Encoders/SimpleAudioEncoder.cpp
--- a/Engine/Plugins/Experimental/AVCodecs/AVCodecsCore/Source/AVCodecsCoreRHI/Private/Audio/Encoders/SimpleAudioEncoder.cpp
+++ b/Engine/Plugins/Experimental/AVCodecs/AVCodecsCore/Source/AVCodecsCoreRHI/Private/Audio/Encoders/SimpleAudioEncoder.cpp
@@ -7,6 +7,8 @@
 
 #include "HAL/RunnableThread.h"
 
+#include UE_INLINE_GENERATED_CPP_BY_NAME(SimpleAudioEncoder)
+
 uint32 USimpleAudioEncoder::Run()
 {
 	while (IsOpen() && IsAsync())
diff --git a/Engine/Plugins/Experimental/AVCodecs/AVCodecsCore/Source/AVCodecsCoreRHI/Private/Audio/SimpleAudio.cpp b/Engine/Plugins/Experimental/AVCodecs/AVCodecsCore/Source/AVCodecsCoreRHI/Private/Audio/SimpleAudio.cpp
--- a/Engine/Plugins/Experimental/AVCodecs/AVCodecsCore/Source/AVCodecsCoreRHI/Private/Audio/SimpleAudio.cpp
+++ b/Engine/Plugins/Experimental/AVCodecs/AVCodecsCore/Source/AVCodecsCoreRHI/Private/Audio/SimpleAudio.cpp
@@ -4,6 +4,8 @@
 
 #include "Audio/Encoders/Configs/AudioEncoderConfigAAC.h"
 
+#include UE_INLINE_GENERATED_CPP_BY_NAME(SimpleAudio)
+
 ESimpleAudioCodec USimpleAudioHelper::GuessCodec(TSharedRef<FAVInstance> const& Instance)
 {
 	if (Instance->Has<FAudioEncoderConfigAAC>())
diff --git a/Engine/Plugins/Experimental/AVCodecs/AVCodecsCore/Source/AVCodecsCoreRHI/Private/SimpleAV.cpp b/Engine/Plugins/Experimental/AVCodecs/AVCodecsCore/Source/AVCodecsCoreRHI/Private/SimpleAV.cpp
--- a/Engine/Plugins/Experimental/AVCodecs/AVCodecsCore/Source/AVCodecsCoreRHI/Private/SimpleAV.cpp
+++ b/Engine/Plugins/Experimental/AVCodecs/AVCodecsCore/Source/AVCodecsCoreRHI/Private/SimpleAV.cpp
@@ -2,6 +2,8 @@
 
 #include "SimpleAV.h"
 
+#include UE_INLINE_GENERATED_CPP_BY_NAME(SimpleAV)
+
 EAVPreset USimpleAVHelper::ConvertPreset(ESimpleAVPreset From)
 {
 	switch (From)
diff --git a/Engine/Plugins/Experimental/AVCodecs/AVCodecsCore/Source/AVCodecsCoreRHI/Private/Video/Decoders/SimpleVideoDecoder.cpp b/Engine/Plugins/Experimental/AVCodecs/AVCodecsCore/Source/AVCodecsCoreRHI/Private/Video/Decoders/SimpleVideoDecoder.cpp
--- a/Engine/Plugins/Experimental/AVCodecs/AVCodecsCore/Source/AVCodecsCoreRHI/Private/Video/Decoders/SimpleVideoDecoder.cpp
+++ b/Engine/Plugins/Experimental/AVCodecs/AVCodecsCore/Source/AVCodecsCoreRHI/Private/Video/Decoders/SimpleVideoDecoder.cpp
@@ -10,6 +10,8 @@
 #include "Video/Decoders/Configs/VideoDecoderConfigH265.h"
 #include "Video/Resources/VideoResourceRHI.h"
 
+#include UE_INLINE_GENERATED_CPP_BY_NAME(SimpleVideoDecoder)
+
 uint32 USimpleVideoDecoder::Run()
 {
 	while (IsOpen() && IsAsync())
diff --git a/Engine/Plugins/Experimental/AVCodecs/AVCodecsCore/Source/AVCodecsCoreRHI/Private/Video/Encoders/SimpleVideoEncoder.cpp b/Engine/Plugins/Experimental/AVCodecs/AVCodecsCore/Source/AVCodecsCoreRHI/Private/Video/Encoders/SimpleVideoEncoder.cpp
--- a/Engine/Plugins/Experimental/AVCodecs/AVCodecsCore/Source/AVCodecsCoreRHI/Private/Video/Encoders/SimpleVideoEncoder.cpp
+++ b/Engine/Plugins/Experimental/AVCodecs/AVCodecsCore/Source/AVCodecsCoreRHI/Private/Video/Encoders/SimpleVideoEncoder.cpp
@@ -12,6 +12,8 @@
 #include "Video/Encoders/Configs/VideoEncoderConfigH265.h"
 #include "Video/Resources/VideoResourceRHI.h"
 
+#include UE_INLINE_GENERATED_CPP_BY_NAME(SimpleVideoEncoder)
+
 
 uint32 USimpleVideoEncoder::Run()
 {
diff --git a/Engine/Plugins/Experimental/AVCodecs/AVCodecsCore/Source/AVCodecsCoreRHI/Private/Video/SimpleVideo.cpp b/Engine/Plugins/Experimental/AVCodecs/AVCodecsCore/Source/AVCodecsCoreRHI/Private/Video/SimpleVideo.cpp
--- a/Engine/Plugins/Experimental/AVCodecs/AVCodecsCore/Source/AVCodecsCoreRHI/Private/Video/SimpleVideo.cpp
+++ b/Engine/Plugins/Experimental/AVCodecs/AVCodecsCore/Source/AVCodecsCoreRHI/Private/Video/SimpleVideo.cpp
@@ -11,6 +11,8 @@
 
 #include "RenderingThread.h"
 
+#include UE_INLINE_GENERATED_CPP_BY_NAME(SimpleVideo)
+
 ESimpleVideoCodec USimpleVideoHelper::GuessCodec(TSharedRef<FAVInstance> const& Instance)
 {
 	if (Instance->Has<FVideoEncoderConfigH264>() || Instance->Has<FVideoDecoderConfigH264>())
diff --git a/Engine/Plugins/Experimental/AVCodecs/AVCodecsCore/Source/AVCodecsCoreRHI/Public/Audio/Encoders/SimpleAudioEncoder.h b/Engine/Plugins/Experimental/AVCodecs/AVCodecsCore/Source/AVCodecsCoreRHI/Public/Audio/Encoders/SimpleAudioEncoder.h
--- a/Engine/Plugins/Experimental/AVCodecs/AVCodecsCore/Source/AVCodecsCoreRHI/Public/Audio/Encoders/SimpleAudioEncoder.h
+++ b/Engine/Plugins/Experimental/AVCodecs/AVCodecsCore/Source/AVCodecsCoreRHI/Public/Audio/Encoders/SimpleAudioEncoder.h
@@ -13,6 +13,8 @@
 
 #include "SimpleAudioEncoder.generated.h"
 
+#define UE_API AVCODECSCORERHI_API
+
 USTRUCT(BlueprintType)
 struct FSimpleAudioEncoderConfig
 {
@@ -51,8 +53,8 @@ public:
     }
 };
 
-UCLASS(Blueprintable)
-class AVCODECSCORERHI_API USimpleAudioEncoder : public UObject, public FRunnable
+UCLASS(MinimalAPI, Blueprintable)
+class USimpleAudioEncoder : public UObject, public FRunnable
 {
 	GENERATED_BODY()
 
@@ -76,42 +78,44 @@ private:
 	TQueue<FAsyncFrame> AsyncQueue;
 	TQueue<TSharedPtr<FAudioResourceCPU>> AsyncPool;
 
-	virtual uint32 Run() override;
-	virtual void Exit() override;
+	UE_API virtual uint32 Run() override;
+	UE_API virtual void Exit() override;
 
 public:
 	UFUNCTION(BlueprintCallable, Category = "Audio")
-	bool IsAsync() const;
+	UE_API bool IsAsync() const;
 
 	UFUNCTION(BlueprintCallable, Category = "Audio")
-	bool IsOpen() const;
+	UE_API bool IsOpen() const;
 
 	UFUNCTION(BlueprintCallable, Category = "Audio")
-	UPARAM(DisplayName = "Was Success") bool Open(ESimpleAudioCodec Codec, FSimpleAudioEncoderConfig Config, bool bAsynchronous);
+	UE_API UPARAM(DisplayName = "Was Success") bool Open(ESimpleAudioCodec Codec, FSimpleAudioEncoderConfig Config, bool bAsynchronous);
 
 	UFUNCTION(BlueprintCallable, Category = "Audio")
-	void Close();
+	UE_API void Close();
 	
 	UFUNCTION(BlueprintCallable, Category = "Audio", meta = (DisplayName = "Send Frame (PCM16)"))
-	UPARAM(DisplayName = "Was Success") bool SendFrameFloat(TArray<float> const& Resource, double Timestamp, int32 NumSamples, float SampleDuration);
+	UE_API UPARAM(DisplayName = "Was Success") bool SendFrameFloat(TArray<float> const& Resource, double Timestamp, int32 NumSamples, float SampleDuration);
 
-	bool SendFrame(Audio::TSampleBuffer<float> const& Resource, double Timestamp);
-	bool SendFrame(float const* ResourceData, double Timestamp, int32 NumSamples, float SampleDuration);
+	UE_API bool SendFrame(Audio::TSampleBuffer<float> const& Resource, double Timestamp);
+	UE_API bool SendFrame(float const* ResourceData, double Timestamp, int32 NumSamples, float SampleDuration);
 
 	UFUNCTION(BlueprintCallable, Category = "Audio")
-	UPARAM(DisplayName = "Was Success") bool ReceivePacket(FSimpleAudioPacket& OutPacket);
+	UE_API UPARAM(DisplayName = "Was Success") bool ReceivePacket(FSimpleAudioPacket& OutPacket);
 
 	UFUNCTION(BlueprintCallable, Category = "Audio")
-	void ReceivePackets(TArray<FSimpleAudioPacket>& OutPackets);
+	UE_API void ReceivePackets(TArray<FSimpleAudioPacket>& OutPackets);
 
 	// Configuration
 public:
 	UFUNCTION(BlueprintCallable, Category = "Audio")
-	ESimpleAudioCodec GetCodec() const;
+	UE_API ESimpleAudioCodec GetCodec() const;
 
 	UFUNCTION(BlueprintCallable, Category = "Audio")
-	FSimpleAudioEncoderConfig GetConfig() const;
+	UE_API FSimpleAudioEncoderConfig GetConfig() const;
 	
 	UFUNCTION(BlueprintCallable, Category = "Audio")
-	void SetConfig(FSimpleAudioEncoderConfig NewConfig);
+	UE_API void SetConfig(FSimpleAudioEncoderConfig NewConfig);
 };
+
+#undef UE_API
diff --git a/Engine/Plugins/Experimental/AVCodecs/AVCodecsCore/Source/AVCodecsCoreRHI/Public/Audio/SimpleAudio.h b/Engine/Plugins/Experimental/AVCodecs/AVCodecsCore/Source/AVCodecsCoreRHI/Public/Audio/SimpleAudio.h
--- a/Engine/Plugins/Experimental/AVCodecs/AVCodecsCore/Source/AVCodecsCoreRHI/Public/Audio/SimpleAudio.h
+++ b/Engine/Plugins/Experimental/AVCodecs/AVCodecsCore/Source/AVCodecsCoreRHI/Public/Audio/SimpleAudio.h
@@ -9,6 +9,8 @@
 
 #include "SimpleAudio.generated.h"
 
+#define UE_API AVCODECSCORERHI_API
+
 UENUM(BlueprintType)
 enum class ESimpleAudioCodec : uint8
 {
@@ -24,11 +26,13 @@ public:
 	FAudioPacket RawPacket;
 };
 
-UCLASS(Abstract)
-class AVCODECSCORERHI_API USimpleAudioHelper : public USimpleAVHelper
+UCLASS(MinimalAPI, Abstract)
+class USimpleAudioHelper : public USimpleAVHelper
 {
 	GENERATED_BODY()
 
 public:
-	static ESimpleAudioCodec GuessCodec(TSharedRef<FAVInstance> const& Instance);
+	static UE_API ESimpleAudioCodec GuessCodec(TSharedRef<FAVInstance> const& Instance);
 };
+
+#undef UE_API
diff --git a/Engine/Plugins/Experimental/AVCodecs/AVCodecsCore/Source/AVCodecsCoreRHI/Public/SimpleAV.h b/Engine/Plugins/Experimental/AVCodecs/AVCodecsCore/Source/AVCodecsCoreRHI/Public/SimpleAV.h
--- a/Engine/Plugins/Experimental/AVCodecs/AVCodecsCore/Source/AVCodecsCoreRHI/Public/SimpleAV.h
+++ b/Engine/Plugins/Experimental/AVCodecs/AVCodecsCore/Source/AVCodecsCoreRHI/Public/SimpleAV.h
@@ -6,6 +6,8 @@
 
 #include "SimpleAV.generated.h"
 
+#define UE_API AVCODECSCORERHI_API
+
 UENUM(BlueprintType)
 enum class ESimpleAVPreset : uint8
 {
@@ -17,11 +19,13 @@ enum class ESimpleAVPreset : uint8
 	Lossless,
 };
 
-UCLASS(Abstract)
-class AVCODECSCORERHI_API USimpleAVHelper : public UObject
+UCLASS(MinimalAPI, Abstract)
+class USimpleAVHelper : public UObject
 {
 	GENERATED_BODY()
 
 public:
-	static EAVPreset ConvertPreset(ESimpleAVPreset From);
+	static UE_API EAVPreset ConvertPreset(ESimpleAVPreset From);
 };
+
+#undef UE_API
diff --git a/Engine/Plugins/Experimental/AVCodecs/AVCodecsCore/Source/AVCodecsCoreRHI/Public/Video/Decoders/SimpleVideoDecoder.h b/Engine/Plugins/Experimental/AVCodecs/AVCodecsCore/Source/AVCodecsCoreRHI/Public/Video/Decoders/SimpleVideoDecoder.h
--- a/Engine/Plugins/Experimental/AVCodecs/AVCodecsCore/Source/AVCodecsCoreRHI/Public/Video/Decoders/SimpleVideoDecoder.h
+++ b/Engine/Plugins/Experimental/AVCodecs/AVCodecsCore/Source/AVCodecsCoreRHI/Public/Video/Decoders/SimpleVideoDecoder.h
@@ -11,8 +11,10 @@
 
 #include "SimpleVideoDecoder.generated.h"
 
-UCLASS(Blueprintable)
-class AVCODECSCORERHI_API USimpleVideoDecoder : public UObject, public FRunnable
+#define UE_API AVCODECSCORERHI_API
+
+UCLASS(MinimalAPI, Blueprintable)
+class USimpleVideoDecoder : public UObject, public FRunnable
 {
 	GENERATED_BODY()
 
@@ -22,30 +24,32 @@ private:
 	FRunnableThread* AsyncThread = nullptr;
 	TQueue<FSimpleVideoPacket> AsyncQueue;
 
-	virtual uint32 Run() override;
-	virtual void Exit() override;
+	UE_API virtual uint32 Run() override;
+	UE_API virtual void Exit() override;
 
 public:
 	UFUNCTION(BlueprintCallable, Category = "Video")
-	bool IsAsync() const;
+	UE_API bool IsAsync() const;
 
 	UFUNCTION(BlueprintCallable, Category = "Video")
-	bool IsOpen() const;
+	UE_API bool IsOpen() const;
 
 	UFUNCTION(BlueprintCallable, Category = "Video")
-	UPARAM(DisplayName = "Was Success") bool Open(ESimpleVideoCodec Codec, bool bAsynchronous);
+	UE_API UPARAM(DisplayName = "Was Success") bool Open(ESimpleVideoCodec Codec, bool bAsynchronous);
 
 	UFUNCTION(BlueprintCallable, Category = "Video")
-	void Close();
+	UE_API void Close();
 
 	UFUNCTION(BlueprintCallable, Category = "Video")
-	UPARAM(DisplayName = "Was Success") bool SendPacket(FSimpleVideoPacket const& Packet);
+	UE_API UPARAM(DisplayName = "Was Success") bool SendPacket(FSimpleVideoPacket const& Packet);
 
 	UFUNCTION(BlueprintCallable, Category = "Video")
-	UPARAM(DisplayName = "Was Success") bool ReceiveFrame(UTextureRenderTarget2D* Resource);
+	UE_API UPARAM(DisplayName = "Was Success") bool ReceiveFrame(UTextureRenderTarget2D* Resource);
 
 	// Configuration
 public:
 	UFUNCTION(BlueprintCallable, Category = "Video")
-	ESimpleVideoCodec GetCodec() const;
+	UE_API ESimpleVideoCodec GetCodec() const;
 };
+
+#undef UE_API
diff --git a/Engine/Plugins/Experimental/AVCodecs/AVCodecsCore/Source/AVCodecsCoreRHI/Public/Video/Encoders/SimpleVideoEncoder.h b/Engine/Plugins/Experimental/AVCodecs/AVCodecsCore/Source/AVCodecsCoreRHI/Public/Video/Encoders/SimpleVideoEncoder.h
--- a/Engine/Plugins/Experimental/AVCodecs/AVCodecsCore/Source/AVCodecsCoreRHI/Public/Video/Encoders/SimpleVideoEncoder.h
+++ b/Engine/Plugins/Experimental/AVCodecs/AVCodecsCore/Source/AVCodecsCoreRHI/Public/Video/Encoders/SimpleVideoEncoder.h
@@ -13,6 +13,8 @@
 
 #include "SimpleVideoEncoder.generated.h"
 
+#define UE_API AVCODECSCORERHI_API
+
 USTRUCT(BlueprintType)
 struct FSimpleVideoEncoderConfig
 {
@@ -61,8 +63,8 @@ public:
     }
 };
 
-UCLASS(Blueprintable)
-class AVCODECSCORERHI_API USimpleVideoEncoder : public UObject, public FRunnable
+UCLASS(MinimalAPI, Blueprintable)
+class USimpleVideoEncoder : public UObject, public FRunnable
 {
 	GENERATED_BODY()
 
@@ -88,44 +90,46 @@ private:
 	TQueue<FAsyncFrame> AsyncQueue;
 	TQueue<TSharedPtr<FVideoResourceRHI>> AsyncPool;
 
-	virtual uint32 Run() override;
-	virtual void Exit() override;
+	UE_API virtual uint32 Run() override;
+	UE_API virtual void Exit() override;
 
 public:
 	UFUNCTION(BlueprintCallable, Category = "Video")
-	bool IsAsync() const;
+	UE_API bool IsAsync() const;
 
 	UFUNCTION(BlueprintCallable, Category = "Video")
-	bool IsOpen() const;
+	UE_API bool IsOpen() const;
 
 	UFUNCTION(BlueprintCallable, Category = "Video")
-	UPARAM(DisplayName = "Was Success") bool Open(ESimpleVideoCodec Codec, FSimpleVideoEncoderConfig Config, bool bAsynchronous);
+	UE_API UPARAM(DisplayName = "Was Success") bool Open(ESimpleVideoCodec Codec, FSimpleVideoEncoderConfig Config, bool bAsynchronous);
 
 	UFUNCTION(BlueprintCallable, Category = "Video")
-	void Close();
+	UE_API void Close();
 
 	UFUNCTION(BlueprintCallable, Category = "Video", meta = (DisplayName = "Send Frame (Render Target)"))
-	UPARAM(DisplayName = "Was Success") bool SendFrameRenderTarget(UTextureRenderTarget2D* Resource, double Timestamp, bool bForceKeyframe = false);
+	UE_API UPARAM(DisplayName = "Was Success") bool SendFrameRenderTarget(UTextureRenderTarget2D* Resource, double Timestamp, bool bForceKeyframe = false);
 	
 	UFUNCTION(BlueprintCallable, Category = "Video", meta = (DisplayName = "Send Frame (Texture)"))
-	UPARAM(DisplayName = "Was Success") bool SendFrameTexture(UTexture2D* Resource, double Timestamp, bool bForceKeyframe = false);
+	UE_API UPARAM(DisplayName = "Was Success") bool SendFrameTexture(UTexture2D* Resource, double Timestamp, bool bForceKeyframe = false);
 
-	bool SendFrame(FTextureRHIRef const& Resource, double Timestamp, bool bForceKeyframe = false);
+	UE_API bool SendFrame(FTextureRHIRef const& Resource, double Timestamp, bool bForceKeyframe = false);
 
 	UFUNCTION(BlueprintCallable, Category = "Video")
-	UPARAM(DisplayName = "Was Success") bool ReceivePacket(FSimpleVideoPacket& OutPacket);
+	UE_API UPARAM(DisplayName = "Was Success") bool ReceivePacket(FSimpleVideoPacket& OutPacket);
 
 	UFUNCTION(BlueprintCallable, Category = "Video")
-	void ReceivePackets(TArray<FSimpleVideoPacket>& OutPackets);
+	UE_API void ReceivePackets(TArray<FSimpleVideoPacket>& OutPackets);
 
 	// Configuration
 public:
 	UFUNCTION(BlueprintCallable, Category = "Video")
-	ESimpleVideoCodec GetCodec() const;
+	UE_API ESimpleVideoCodec GetCodec() const;
 
 	UFUNCTION(BlueprintCallable, Category = "Video")
-	FSimpleVideoEncoderConfig GetConfig() const;
+	UE_API FSimpleVideoEncoderConfig GetConfig() const;
 	
 	UFUNCTION(BlueprintCallable, Category = "Video")
-	void SetConfig(FSimpleVideoEncoderConfig NewConfig);
+	UE_API void SetConfig(FSimpleVideoEncoderConfig NewConfig);
 };
+
+#undef UE_API
diff --git a/Engine/Plugins/Experimental/AVCodecs/AVCodecsCore/Source/AVCodecsCoreRHI/Public/Video/Resources/VideoResourceRHI.h b/Engine/Plugins/Experimental/AVCodecs/AVCodecsCore/Source/AVCodecsCoreRHI/Public/Video/Resources/VideoResourceRHI.h
--- a/Engine/Plugins/Experimental/AVCodecs/AVCodecsCore/Source/AVCodecsCoreRHI/Public/Video/Resources/VideoResourceRHI.h
+++ b/Engine/Plugins/Experimental/AVCodecs/AVCodecsCore/Source/AVCodecsCoreRHI/Public/Video/Resources/VideoResourceRHI.h
@@ -14,11 +14,11 @@ THIRD_PARTY_INCLUDES_START
 THIRD_PARTY_INCLUDES_END
 #endif
 
-class AVCODECSCORERHI_API FVideoContextRHI : public FAVContext
+class FVideoContextRHI : public FAVContext
 {
 };
 
-class AVCODECSCORERHI_API FVideoResourceRHI : public TVideoResource<FVideoContextRHI>
+class FVideoResourceRHI : public TVideoResource<FVideoContextRHI>
 {
 public:
 	struct FRawData
@@ -36,40 +36,40 @@ private:
 	FRawData Raw;
 
 public:
-	static FAVLayout GetLayoutFrom(TSharedRef<FAVDevice> const& Device, FTextureRHIRef const& Raw);
-	static FVideoDescriptor GetDescriptorFrom(TSharedRef<FAVDevice> const& Device, FTextureRHIRef const& Raw);
+	static AVCODECSCORERHI_API FAVLayout GetLayoutFrom(TSharedRef<FAVDevice> const& Device, FTextureRHIRef const& Raw);
+	static AVCODECSCORERHI_API FVideoDescriptor GetDescriptorFrom(TSharedRef<FAVDevice> const& Device, FTextureRHIRef const& Raw);
 
-	static TSharedPtr<FVideoResourceRHI> Create(TSharedPtr<FAVDevice> const& Device, FVideoDescriptor const& Descriptor, ETextureCreateFlags AdditionalFlags = ETextureCreateFlags::None, bool bIsSRGB = true);
+	static AVCODECSCORERHI_API TSharedPtr<FVideoResourceRHI> Create(TSharedPtr<FAVDevice> const& Device, FVideoDescriptor const& Descriptor, ETextureCreateFlags AdditionalFlags = ETextureCreateFlags::None, bool bIsSRGB = true);
 
 
 	inline FRawData const& GetRaw() const { return Raw; }
 
-	FVideoResourceRHI(TSharedRef<FAVDevice> const& Device, FRawData const& Raw, FVideoDescriptor const& OverrideDescriptor);
-	FVideoResourceRHI(TSharedRef<FAVDevice> const& Device, FRawData const& Raw);
+	AVCODECSCORERHI_API FVideoResourceRHI(TSharedRef<FAVDevice> const& Device, FRawData const& Raw, FVideoDescriptor const& OverrideDescriptor);
+	AVCODECSCORERHI_API FVideoResourceRHI(TSharedRef<FAVDevice> const& Device, FRawData const& Raw);
 	virtual ~FVideoResourceRHI() override = default;
 
-	virtual FAVResult Validate() const override;
+	AVCODECSCORERHI_API virtual FAVResult Validate() const override;
 
-	virtual void Lock() override;
-	virtual FScopeLock LockScope() override;
+	AVCODECSCORERHI_API virtual void Lock() override;
+	AVCODECSCORERHI_API virtual FScopeLock LockScope() override;
 
-	void CopyFrom(TArrayView64<uint8> const& From);
-	void CopyFrom(TArrayView<uint8> const& From);
-	void CopyFrom(TArray64<uint8> const& From);
-	void CopyFrom(TArray<uint8> const& From);
-	void CopyFrom(FTextureRHIRef const& From);
+	AVCODECSCORERHI_API void CopyFrom(TArrayView64<uint8> const& From);
+	AVCODECSCORERHI_API void CopyFrom(TArrayView<uint8> const& From);
+	AVCODECSCORERHI_API void CopyFrom(TArray64<uint8> const& From);
+	AVCODECSCORERHI_API void CopyFrom(TArray<uint8> const& From);
+	AVCODECSCORERHI_API void CopyFrom(FTextureRHIRef const& From);
 
-	TSharedPtr<FVideoResourceRHI> TransformResource(FVideoDescriptor const& OutDescriptor);
-	void TransformResourceTo(FRHICommandListImmediate& RHICmdList, FTextureRHIRef Target, FVideoDescriptor const& OutDescriptor);
+	AVCODECSCORERHI_API TSharedPtr<FVideoResourceRHI> TransformResource(FVideoDescriptor const& OutDescriptor);
+	AVCODECSCORERHI_API void TransformResourceTo(FRHICommandListImmediate& RHICmdList, FTextureRHIRef Target, FVideoDescriptor const& OutDescriptor);
 };
 
-class AVCODECSCORERHI_API FResolvableVideoResourceRHI : public TResolvableVideoResource<FVideoResourceRHI>
+class FResolvableVideoResourceRHI : public TResolvableVideoResource<FVideoResourceRHI>
 {
 public:
 	ETextureCreateFlags AdditionalFlags = ETextureCreateFlags::None;
 	
 protected:
-	virtual TSharedPtr<FVideoResourceRHI> TryResolve(TSharedPtr<FAVDevice> const& Device, FVideoDescriptor const& Descriptor) override;
+	AVCODECSCORERHI_API virtual TSharedPtr<FVideoResourceRHI> TryResolve(TSharedPtr<FAVDevice> const& Device, FVideoDescriptor const& Descriptor) override;
 };
 
 /*
diff --git a/Engine/Plugins/Experimental/AVCodecs/AVCodecsCore/Source/AVCodecsCoreRHI/Public/Video/SimpleVideo.h b/Engine/Plugins/Experimental/AVCodecs/AVCodecsCore/Source/AVCodecsCoreRHI/Public/Video/SimpleVideo.h
--- a/Engine/Plugins/Experimental/AVCodecs/AVCodecsCore/Source/AVCodecsCoreRHI/Public/Video/SimpleVideo.h
+++ b/Engine/Plugins/Experimental/AVCodecs/AVCodecsCore/Source/AVCodecsCoreRHI/Public/Video/SimpleVideo.h
@@ -9,6 +9,8 @@
 
 #include "SimpleVideo.generated.h"
 
+#define UE_API AVCODECSCORERHI_API
+
 UENUM(BlueprintType)
 enum class ESimpleVideoCodec : uint8
 {
@@ -25,14 +27,16 @@ public:
 	FVideoPacket RawPacket;
 };
 
-UCLASS(Abstract)
-class AVCODECSCORERHI_API USimpleVideoHelper : public USimpleAVHelper
+UCLASS(MinimalAPI, Abstract)
+class USimpleVideoHelper : public USimpleAVHelper
 {
 	GENERATED_BODY()
 
 public:
-	static ESimpleVideoCodec GuessCodec(TSharedRef<FAVInstance> const& Instance);
+	static UE_API ESimpleVideoCodec GuessCodec(TSharedRef<FAVInstance> const& Instance);
 
 	UFUNCTION(BlueprintCallable, Category = "Video", meta = (WorldContext = "WorldContext"))
-	static void ShareRenderTarget2D(class UTextureRenderTarget2D* RenderTarget);
+	static UE_API void ShareRenderTarget2D(class UTextureRenderTarget2D* RenderTarget);
 };
+
+#undef UE_API
diff --git a/Engine/Plugins/Experimental/AVCodecs/LibVpxCodecs/Source/LibVpxCodecs/Public/LibVpx.h b/Engine/Plugins/Experimental/AVCodecs/LibVpxCodecs/Source/LibVpxCodecs/Public/LibVpx.h
--- a/Engine/Plugins/Experimental/AVCodecs/LibVpxCodecs/Source/LibVpxCodecs/Public/LibVpx.h
+++ b/Engine/Plugins/Experimental/AVCodecs/LibVpxCodecs/Source/LibVpxCodecs/Public/LibVpx.h
@@ -17,14 +17,18 @@ THIRD_PARTY_INCLUDES_START
 #include "vpx_frame_buffer.h"
 #include "vpx_image.h"
 #include "vpx_integer.h"
+
+#define UE_API LIBVPXCODECS_API
 THIRD_PARTY_INCLUDES_END
 
-class LIBVPXCODECS_API FLibVpx : public FAPI
+class FLibVpx : public FAPI
 {
 public:
 	FLibVpx() = default;
 
-	virtual bool IsValid() const override;
+	UE_API virtual bool IsValid() const override;
 };
 
 DECLARE_TYPEID(FLibVpx, LIBVPXCODECS_API);
+
+#undef UE_API
diff --git a/Engine/Plugins/Experimental/AVCodecs/LibVpxCodecs/Source/LibVpxCodecs/Public/Video/Decoders/Configs/VideoDecoderConfigLibVpx.h b/Engine/Plugins/Experimental/AVCodecs/LibVpxCodecs/Source/LibVpxCodecs/Public/Video/Decoders/Configs/VideoDecoderConfigLibVpx.h
--- a/Engine/Plugins/Experimental/AVCodecs/LibVpxCodecs/Source/LibVpxCodecs/Public/Video/Decoders/Configs/VideoDecoderConfigLibVpx.h
+++ b/Engine/Plugins/Experimental/AVCodecs/LibVpxCodecs/Source/LibVpxCodecs/Public/Video/Decoders/Configs/VideoDecoderConfigLibVpx.h
@@ -6,7 +6,7 @@
 #include "Video/VideoDecoder.h"
 #include "LibVpx.h"
 
-struct LIBVPXCODECS_API FVideoDecoderConfigLibVpx : public FAVConfig
+struct FVideoDecoderConfigLibVpx : public FAVConfig
 {
 public:
 	// Values are parsed from the bitstream during decoding
diff --git a/Engine/Plugins/Experimental/AVCodecs/LibVpxCodecs/Source/LibVpxCodecs/Public/Video/Encoders/Configs/VideoEncoderConfigLibVpx.h b/Engine/Plugins/Experimental/AVCodecs/LibVpxCodecs/Source/LibVpxCodecs/Public/Video/Encoders/Configs/VideoEncoderConfigLibVpx.h
--- a/Engine/Plugins/Experimental/AVCodecs/LibVpxCodecs/Source/LibVpxCodecs/Public/Video/Encoders/Configs/VideoEncoderConfigLibVpx.h
+++ b/Engine/Plugins/Experimental/AVCodecs/LibVpxCodecs/Source/LibVpxCodecs/Public/Video/Encoders/Configs/VideoEncoderConfigLibVpx.h
@@ -10,7 +10,7 @@
 
 using namespace UE::AVCodecCore::VP9;
 
-struct LIBVPXCODECS_API FVideoEncoderConfigLibVpx : public FAVConfig
+struct FVideoEncoderConfigLibVpx : public FAVConfig
 {
 public:
 	uint32			 Width = 1920;
