[AVCodecs] Add: Support for iOS and Android

#Synced-CL 48443255

diff --git a/Engine/Plugins/Experimental/AVCodecs/AVCodecsCore/AVCodecsCore.uplugin b/Engine/Plugins/Experimental/AVCodecs/AVCodecsCore/AVCodecsCore.uplugin
--- a/Engine/Plugins/Experimental/AVCodecs/AVCodecsCore/AVCodecsCore.uplugin
+++ b/Engine/Plugins/Experimental/AVCodecs/AVCodecsCore/AVCodecsCore.uplugin
@@ -18,7 +18,7 @@
 			"Name": "AVCodecsCore",
 			"Type": "Runtime",
 			"LoadingPhase": "PostEngineInit",
-			"PlatformAllowList": [ "Win64",	"Linux", "Mac" ]
+			"PlatformAllowList": [ "Win64", "Linux", "Mac", "Android", "IOS" ]
 		},
 		{
 			"Name": "AVCodecsCoreRHI",
@@ -27,12 +27,14 @@
 			"TargetDenyList": [
 				"Program"
 			],
-			"PlatformAllowList": [ "Win64",	"Linux", "Mac"	]
+			"PlatformAllowList": [ "Win64", "Linux", "Mac", "Android", "IOS" ]
 		}
 	],
 	"SupportedTargetPlatforms": [
 		"Win64",
 		"Linux",
-		"Mac"
+		"Mac",
+		"Android", 
+		"IOS"
 	]
 }
\ No newline at end of file
diff --git a/Engine/Plugins/Experimental/AVCodecs/AVCodecsCore/Source/AVCodecsCore/AVCodecsCore.Build.cs b/Engine/Plugins/Experimental/AVCodecs/AVCodecsCore/Source/AVCodecsCore/AVCodecsCore.Build.cs
--- a/Engine/Plugins/Experimental/AVCodecs/AVCodecsCore/Source/AVCodecsCore/AVCodecsCore.Build.cs
+++ b/Engine/Plugins/Experimental/AVCodecs/AVCodecsCore/Source/AVCodecsCore/AVCodecsCore.Build.cs
@@ -4,7 +4,7 @@ using UnrealBuildTool;
 using System.IO;
 using System;
 
-[SupportedPlatforms("Win64", "Linux", "Mac")]
+[SupportedPlatforms("Win64", "Linux", "Mac", "Android", "IOS")]
 public class AVCodecsCore : ModuleRules
 {
     public AVCodecsCore(ReadOnlyTargetRules Target) : base(Target)
diff --git a/Engine/Plugins/Experimental/AVCodecs/AVCodecsCore/Source/AVCodecsCore/Private/Video/Resources/Metal/VideoResourceMetal.cpp b/Engine/Plugins/Experimental/AVCodecs/AVCodecsCore/Source/AVCodecsCore/Private/Video/Resources/Metal/VideoResourceMetal.cpp
--- a/Engine/Plugins/Experimental/AVCodecs/AVCodecsCore/Source/AVCodecsCore/Private/Video/Resources/Metal/VideoResourceMetal.cpp
+++ b/Engine/Plugins/Experimental/AVCodecs/AVCodecsCore/Source/AVCodecsCore/Private/Video/Resources/Metal/VideoResourceMetal.cpp
@@ -26,6 +26,51 @@ static TAVResult<EVideoFormat> ConvertFormat(MTL::PixelFormat Format)
 	}
 }
 
+static TAVResult<uint32> GetBytesPerPixelFromOSType(OSType Format)
+{
+	switch (Format)
+	{
+		case kCVPixelFormatType_OneComponent8:
+		case kCVPixelFormatType_420YpCbCr8BiPlanarVideoRange:
+		case kCVPixelFormatType_420YpCbCr8BiPlanarFullRange:
+			return 1; 	
+		case kCVPixelFormatType_OneComponent16Half:
+		case kCVPixelFormatType_OneComponent16:
+		case kCVPixelFormatType_422YpCbCr8:         
+		case kCVPixelFormatType_422YpCbCr8_yuvs:    
+		case kCVPixelFormatType_422YpCbCr10:        
+		case kCVPixelFormatType_420YpCbCr10BiPlanarVideoRange:
+		case kCVPixelFormatType_420YpCbCr10BiPlanarFullRange:
+		case kCVPixelFormatType_DepthFloat16:
+		case kCVPixelFormatType_14Bayer_GRBG:
+		case kCVPixelFormatType_14Bayer_RGGB:
+		case kCVPixelFormatType_14Bayer_BGGR:
+		case kCVPixelFormatType_14Bayer_GBRG:
+		case kCVPixelFormatType_DisparityFloat16:
+			return 2;
+		case kCVPixelFormatType_24RGB:
+			return 3;
+		case kCVPixelFormatType_OneComponent32Float:
+		case kCVPixelFormatType_TwoComponent16Half:
+		case kCVPixelFormatType_TwoComponent16:
+		case kCVPixelFormatType_30RGB:
+		case kCVPixelFormatType_32RGBA:
+		case kCVPixelFormatType_32BGRA:
+		case kCVPixelFormatType_32ARGB:
+		case kCVPixelFormatType_32ABGR:
+		case kCVPixelFormatType_DepthFloat32:
+		case kCVPixelFormatType_DisparityFloat32:
+			return 4;
+		case kCVPixelFormatType_48RGB: 
+			return 6;
+		case kCVPixelFormatType_TwoComponent32Float:
+			return 8;
+		default:
+			return FAVResult(EAVResult::ErrorUnsupported, TEXT("PixelFormat format is not supported"), TEXT("Metal"));
+
+	}
+}
+
 FVideoContextMetal::FVideoContextMetal(MTL::Device* Device)
 	: Device(Device)
 {
@@ -33,10 +78,10 @@ FVideoContextMetal::FVideoContextMetal(MTL::Device* Device)
 
 FVideoDescriptor FVideoResourceMetal::GetDescriptorFrom(TSharedRef<FAVDevice> const& Device, MTL::Texture* Raw)
 {
-    uint32_t Width = Raw->width();
-    uint32_t Height = Raw->height();
-    TAVResult<EVideoFormat> ConvertedFormat = ConvertFormat(Raw->pixelFormat());
-    
+	uint32_t Width = Raw->width();
+	uint32_t Height = Raw->height();
+	TAVResult<EVideoFormat> ConvertedFormat = ConvertFormat(Raw->pixelFormat());
+	
 	return FVideoDescriptor(ConvertedFormat, Width, Height);
 }
 
@@ -52,12 +97,49 @@ FVideoResourceMetal::~FVideoResourceMetal()
 
 FAVResult FVideoResourceMetal::CopyFrom(CVPixelBufferRef Other)
 {	
+	// TODO (william.belcher): Support planar pixel buffers
+	if (CVPixelBufferIsPlanar(Other))
+	{
+		return FAVResult(EAVResult::ErrorUnsupported, TEXT("Planar pixel buffers are not supported for copying to Metal resources"), TEXT("Metal"));
+	}
+
+	const size_t SrcHeight = CVPixelBufferGetHeight(Other);
+	const size_t SrcWidth = CVPixelBufferGetWidth(Other);
+	
+	if (SrcHeight != GetDescriptor().Height || SrcWidth != GetDescriptor().Width)
+	{
+		return FAVResult(EAVResult::Error, TEXT("PixelBuffer dimensions must match Metal resource dimensions!"), TEXT("Metal"));
+	}
+	
+	TAVResult<uint32> BytesPerPixel = GetBytesPerPixelFromOSType(CVPixelBufferGetPixelFormatType(Other));
+	if (BytesPerPixel.IsNotSuccess())
+	{
+		return BytesPerPixel;
+	}
+	
 	if (const CVReturn Result = CVPixelBufferLockBaseAddress(Other, 0); Result != kCVReturnSuccess)
 	{
 		return FAVResult(EAVResult::Error, TEXT("Failed to lock input pixel buffer!"), TEXT("Metal"));
 	}
+
+	const size_t SrcStride = CVPixelBufferGetBytesPerRow(Other);
+	const uint8* SrcPtr = static_cast<const uint8*>(CVPixelBufferGetBaseAddress(Other));
+	
+	const size_t TightRowSize = BytesPerPixel * SrcWidth;
+	
+	TArray<uint8> PackedData;
+	PackedData.SetNumZeroed(TightRowSize * SrcHeight);
+	
+	uint8* DstPtr = PackedData.GetData();
+	for (size_t Row = 0; Row < SrcHeight; Row++)
+	{
+		const uint8* SrcRow = SrcPtr + Row * SrcStride;
+		uint8* DstRow = DstPtr + Row * TightRowSize;
+		
+		FMemory::Memcpy(DstRow, SrcRow, TightRowSize);
+	}
 	
-	Raw->replaceRegion(MTL::Region(0, 0, GetDescriptor().Width, GetDescriptor().Height), 0, 0, CVPixelBufferGetBaseAddress(Other), CVPixelBufferGetBytesPerRow(Other), 0);
+	Raw->replaceRegion(MTL::Region(0, 0, SrcWidth, SrcHeight), 0, 0, PackedData.GetData(), TightRowSize, 0);
 	
 	CVPixelBufferUnlockBaseAddress(Other, 0);
 	
diff --git a/Engine/Plugins/Experimental/AVCodecs/AVCodecsCore/Source/AVCodecsCore/Public/Video/Resources/Metal/VideoResourceMetal.h b/Engine/Plugins/Experimental/AVCodecs/AVCodecsCore/Source/AVCodecsCore/Public/Video/Resources/Metal/VideoResourceMetal.h
--- a/Engine/Plugins/Experimental/AVCodecs/AVCodecsCore/Source/AVCodecsCore/Public/Video/Resources/Metal/VideoResourceMetal.h
+++ b/Engine/Plugins/Experimental/AVCodecs/AVCodecsCore/Source/AVCodecsCore/Public/Video/Resources/Metal/VideoResourceMetal.h
@@ -11,6 +11,7 @@
 #include "Containers/ResourceArray.h"
 
 THIRD_PARTY_INCLUDES_START
+#include <CoreVideo/CoreVideo.h>
 #include "MetalInclude.h"
 THIRD_PARTY_INCLUDES_END
 
diff --git a/Engine/Plugins/Experimental/AVCodecs/LibVpxCodecs/LibVpxCodecs.uplugin b/Engine/Plugins/Experimental/AVCodecs/LibVpxCodecs/LibVpxCodecs.uplugin
--- a/Engine/Plugins/Experimental/AVCodecs/LibVpxCodecs/LibVpxCodecs.uplugin
+++ b/Engine/Plugins/Experimental/AVCodecs/LibVpxCodecs/LibVpxCodecs.uplugin
@@ -24,7 +24,9 @@
 			"PlatformAllowList": [
 				"Win64",
                 "Linux",
-                "Mac"
+                "Mac",
+				"Android",
+				"IOS"
 			]
 		}
 	],
@@ -37,6 +39,8 @@
 	"SupportedTargetPlatforms": [
 		"Win64",
         "Linux",
-        "Mac"
+        "Mac",
+		"Android",
+		"IOS"
 	]
 }
\ No newline at end of file
diff --git a/Engine/Plugins/Experimental/AVCodecs/VTCodecs/Source/VTCodecs/Private/VT.cpp b/Engine/Plugins/Experimental/AVCodecs/VTCodecs/Source/VTCodecs/Private/VT.cpp
--- a/Engine/Plugins/Experimental/AVCodecs/VTCodecs/Source/VTCodecs/Private/VT.cpp
+++ b/Engine/Plugins/Experimental/AVCodecs/VTCodecs/Source/VTCodecs/Private/VT.cpp
@@ -8,7 +8,10 @@
 #include "RHI.h"
 
 THIRD_PARTY_INCLUDES_START
-#include <IOKit/IOKitLib.h>
+#if PLATFORM_MAC
+	// IOKit pull mac only includes
+	#include <IOKit/IOKitLib.h>
+#endif
 THIRD_PARTY_INCLUDES_END
 
 REGISTER_TYPEID(FVT);
@@ -18,6 +21,7 @@ namespace
 	// If IoSurface is unsupported, then VTIsHardwareDecodeSupported will assert internally.	
 	bool IsIoSurfaceAvailable()
 	{
+#if PLATFORM_MAC
 		// IOSurfaceRoot is determined by running `ioreg -n IOSurfaceRoot` which prints
 		// "CFBundleIdentifier" = "com.apple.iokit.IOSurface"
 		// "IOMatchCategory" = "IOSurfaceRoot"
@@ -36,6 +40,7 @@ namespace
 			return false;
 		}
 		IOServiceClose(Connect);
+#endif
 		return true;
 	}
 }
diff --git a/Engine/Plugins/Experimental/AVCodecs/VTCodecs/Source/VTCodecs/Private/VTCodecsModule.cpp b/Engine/Plugins/Experimental/AVCodecs/VTCodecs/Source/VTCodecs/Private/VTCodecsModule.cpp
--- a/Engine/Plugins/Experimental/AVCodecs/VTCodecs/Source/VTCodecs/Private/VTCodecsModule.cpp
+++ b/Engine/Plugins/Experimental/AVCodecs/VTCodecs/Source/VTCodecs/Private/VTCodecsModule.cpp
@@ -24,7 +24,6 @@ public:
 	virtual void StartupModule() override
 	{       
         TFunction<bool(TSharedRef<FAVDevice> const&, const TCHAR*)> CheckCodecSupport = [](TSharedRef<FAVDevice> const& Device, const TCHAR* Codec) {
-			CFStringRef TargetCodec = CFStringCreateWithCString(kCFAllocatorDefault, TCHAR_TO_ANSI(Codec), kCFStringEncodingUTF8);
             CFArrayRef EncoderList = NULL;
             bool bSupportsHardware = false;
 
@@ -34,6 +33,16 @@ public:
                 return false;
             }
 
+			CFStringRef TargetCodec = CFStringCreateWithCString(kCFAllocatorDefault, TCHAR_TO_ANSI(Codec), kCFStringEncodingUTF8);
+			if (!TargetCodec)
+			{
+				if (EncoderList) 
+				{
+					CFRelease(EncoderList);
+				}
+				return false;
+			}
+			
             CFIndex EncoderCount = CFArrayGetCount(EncoderList);
             for (CFIndex i = 0; i < EncoderCount; i++) 
             {
@@ -41,7 +50,13 @@ public:
                 if (EncoderDict)
                 {
 					CFStringRef EncoderCodec = (CFStringRef)CFDictionaryGetValue(EncoderDict, kVTVideoEncoderList_CodecName);
-		 			if (EncoderCodec != TargetCodec)
+					if (!EncoderCodec)
+					{
+						continue;
+					}
+
+					CFComparisonResult CompareResult = CFStringCompare(EncoderCodec, TargetCodec, 0);
+		 			if (CompareResult != kCFCompareEqualTo)
 					{
 						continue;
 					}
@@ -56,7 +71,7 @@ public:
                 CFRelease(EncoderList);
             }
 			
-			if(TargetCodec)
+			if (TargetCodec)
 			{
 				CFRelease(TargetCodec);
 			}
@@ -69,7 +84,7 @@ public:
             ::And<FVideoEncoderConfigVT, FVideoEncoderConfigH264>(
                 [CheckCodecSupport](TSharedRef<FAVDevice> const& NewDevice, TSharedRef<FAVInstance> const& NewInstance)
                 {
-					static bool bSupportsCodec = FAPI::Get<FVT>().IsValid() && CheckCodecSupport(NewDevice, TEXT("H.264"));
+					static bool bSupportsCodec = FAPI::Get<FVT>().IsValid() && (CheckCodecSupport(NewDevice, TEXT("H.264")) || CheckCodecSupport(NewDevice, TEXT("avc1")));
 
 					return bSupportsCodec;
                 });
@@ -84,7 +99,10 @@ public:
 					return bSupportsCodec;
                 });
 
+#if PLATFORM_MAC
+		// TODO(Eden.Harris) VTRegisterSupplementalVideoDecoderIfAvailable is only available on mac
         VTRegisterSupplementalVideoDecoderIfAvailable(kCMVideoCodecType_VP9);
+#endif
 
 		FVideoDecoder::RegisterPermutationsOf<FVideoDecoderVT>
 			::With<FVideoResourceMetal>
diff --git a/Engine/Plugins/Experimental/AVCodecs/VTCodecs/Source/VTCodecsRHI/Private/VTCodecsRHIModule.cpp b/Engine/Plugins/Experimental/AVCodecs/VTCodecs/Source/VTCodecsRHI/Private/VTCodecsRHIModule.cpp
--- a/Engine/Plugins/Experimental/AVCodecs/VTCodecs/Source/VTCodecsRHI/Private/VTCodecsRHIModule.cpp
+++ b/Engine/Plugins/Experimental/AVCodecs/VTCodecs/Source/VTCodecsRHI/Private/VTCodecsRHIModule.cpp
@@ -18,7 +18,12 @@ public:
         {
             FCoreDelegates::OnPostEngineInit.AddLambda([]()
                 {
-                    const_cast<FVT&>(FAPI::Get<FVT>()).bHasCompatibleGPU = IsRHIDeviceApple();
+#if PLATFORM_IOS
+					// On iOS, GRHIVendorId is not set so IsRHIDeviceApple returns false
+					const_cast<FVT&>(FAPI::Get<FVT>()).bHasCompatibleGPU = true;
+#else // PLATFORM_IOS
+					const_cast<FVT&>(FAPI::Get<FVT>()).bHasCompatibleGPU = IsRHIDeviceApple();
+#endif // PLATFORM_IOS
                 });
         }
     }
diff --git a/Engine/Plugins/Experimental/AVCodecs/VTCodecs/VTCodecs.uplugin b/Engine/Plugins/Experimental/AVCodecs/VTCodecs/VTCodecs.uplugin
--- a/Engine/Plugins/Experimental/AVCodecs/VTCodecs/VTCodecs.uplugin
+++ b/Engine/Plugins/Experimental/AVCodecs/VTCodecs/VTCodecs.uplugin
@@ -23,6 +23,7 @@
 			],
 			"PlatformAllowList": [
 				"Mac",
+				"IOS"
 			]
 		},
 		{
@@ -30,10 +31,11 @@
 			"Type": "Runtime",
 			"LoadingPhase": "PostConfigInit",
 			"TargetDenyList": [
-				"Program"
+				"Server"
 			],
 			"PlatformAllowList": [
 				"Mac",
+				"IOS"
 			]
 		}
 	],
@@ -45,5 +47,6 @@
 	],
 	"SupportedTargetPlatforms": [
 		"Mac",
+		"IOS"
 	]
 }
\ No newline at end of file
