[AVCodecs] Fix: Non-unity compile errors on Android

#Synced-CL 48443255

diff --git a/Engine/Plugins/Experimental/AVCodecs/AVCodecsCore/Source/AVCodecsCore/Public/Video/Resources/Vulkan/VideoResourceVulkan.h b/Engine/Plugins/Experimental/AVCodecs/AVCodecsCore/Source/AVCodecsCore/Public/Video/Resources/Vulkan/VideoResourceVulkan.h
--- a/Engine/Plugins/Experimental/AVCodecs/AVCodecsCore/Source/AVCodecsCore/Public/Video/Resources/Vulkan/VideoResourceVulkan.h
+++ b/Engine/Plugins/Experimental/AVCodecs/AVCodecsCore/Source/AVCodecsCore/Public/Video/Resources/Vulkan/VideoResourceVulkan.h
@@ -2,6 +2,8 @@
 
 #pragma once
 
+#if AVCODECS_USE_VULKAN
+
 #include "AVContext.h"
 #include "Video/VideoResource.h"
 
@@ -59,3 +61,5 @@ public:
 
 DECLARE_TYPEID(FVideoContextVulkan, AVCODECSCORE_API);
 DECLARE_TYPEID(FVideoResourceVulkan, AVCODECSCORE_API);
+
+#endif
