Made the Win64 client running in MobileSession mode as a Remote Session + Pixel Streaming 2 host.
Added support for explicit plugin loading at the project/module level, disabling automatic loading by default.

#Synced-CL 50925863

#Synced-CL 50925863

diff --git a/Engine/Plugins/Experimental/AVCodecs/AVCodecsCore/Source/AVCodecsCoreRHI/Private/AVCodecsCoreRHIModule.cpp b/Engine/Plugins/Experimental/AVCodecs/AVCodecsCore/Source/AVCodecsCoreRHI/Private/AVCodecsCoreRHIModule.cpp
--- a/Engine/Plugins/Experimental/AVCodecs/AVCodecsCore/Source/AVCodecsCoreRHI/Private/AVCodecsCoreRHIModule.cpp
+++ b/Engine/Plugins/Experimental/AVCodecs/AVCodecsCore/Source/AVCodecsCoreRHI/Private/AVCodecsCoreRHIModule.cpp
@@ -29,6 +29,8 @@
 	#include "Video/Resources/Metal/VideoResourceMetal.h"
 #endif
 
+#include "Engine/Engine.h"
+
 class FAVCodecsCoreRHI : public IModuleInterface
 {
 public:
@@ -55,52 +57,15 @@ public:
 				IVulkanDynamicRHI::AddEnabledDeviceExtensionsAndLayers(ExtensionsToAdd, TArray<ANSICHAR const*>());
 			}
 #endif // AVCODECS_USE_VULKAN
-			FCoreDelegates::OnPostEngineInit.AddLambda([]()
-			{
-				switch (GDynamicRHI->GetInterfaceType())
-				{
-#if AVCODECS_USE_VULKAN
-				case ERHIInterfaceType::Vulkan:
-					FAVDevice::GetHardwareDevice()->SetContext<FVideoContextVulkan>(
-						MakeShared<FVideoContextVulkan>(
-							GetIVulkanDynamicRHI()->RHIGetVkInstance(),
-							GetIVulkanDynamicRHI()->RHIGetVkDevice(),
-							GetIVulkanDynamicRHI()->RHIGetVkPhysicalDevice(),
-							[] (char const* Name) -> PFN_vkVoidFunction {
-								return (PFN_vkVoidFunction)GetIVulkanDynamicRHI()->RHIGetVkDeviceProcAddr(Name);
-							}));
-				
-					break;
-#endif
-#if AVCODECS_USE_D3D
-				case ERHIInterfaceType::D3D11:
-					FAVDevice::GetHardwareDevice()->SetContext<FVideoContextD3D11>(
-						MakeShared<FVideoContextD3D11>(
-							GetID3D11DynamicRHI()->RHIGetDevice()));
-				
-					break;
-				case ERHIInterfaceType::D3D12:
-					FAVDevice::GetHardwareDevice()->SetContext<FVideoContextD3D12>(
-						MakeShared<FVideoContextD3D12>(
-							GetID3D12DynamicRHI()->RHIGetDevice_NoMGPU()));
-				
-					break;
-#endif
-#if AVCODECS_USE_METAL
-                case ERHIInterfaceType::Metal:
-                    FAVDevice::GetHardwareDevice()->SetContext<FVideoContextMetal>(
-                        MakeShared<FVideoContextMetal>(
-							GetIMetalDynamicRHI()->RHIGetDevice()
-						)
-					);
-                    break;
-#endif
-				default:
-					break;
-				}
 
-				FAVDevice::GetHardwareDevice()->SetContext<FVideoContextRHI>(MakeShared<FVideoContextRHI>());
-			});
+			if (!GEngine || !GEngine->IsInitialized())
+			{
+				PostEngineInitHandle = FCoreDelegates::OnPostEngineInit.AddStatic(&FAVCodecsCoreRHI::PostEngineInit);
+			}
+			else
+			{
+				FAVCodecsCoreRHI::PostEngineInit();
+			}
 		}
 
 		FVideoEncoder::Register<TVideoEncoderRHI<FVideoEncoderConfigAV1>, FVideoResourceRHI, FVideoEncoderConfigAV1>();
@@ -112,6 +77,73 @@ public:
         FVideoDecoder::Register<TVideoDecoderRHI<FVideoDecoderConfigVP9>, FVideoResourceRHI, FVideoDecoderConfigVP9>();
 		FVideoDecoder::Register<TVideoDecoderRHI<FVideoDecoderConfigAV1>, FVideoResourceRHI, FVideoDecoderConfigAV1>();
 	}
+
+	void ShutdownModule() override
+	{
+		// Remove the binding
+		if (PostEngineInitHandle.IsValid())
+		{
+			FCoreDelegates::OnPostEngineInit.Remove(PostEngineInitHandle);
+			PostEngineInitHandle.Reset();
+		}
+	}
+
+	void static PostEngineInit()
+	{
+		// Remove the binding
+		if (PostEngineInitHandle.IsValid())
+		{
+			FCoreDelegates::OnPostEngineInit.Remove(PostEngineInitHandle);
+			PostEngineInitHandle.Reset();
+		}
+
+		switch (GDynamicRHI->GetInterfaceType())
+		{
+#if AVCODECS_USE_VULKAN
+		case ERHIInterfaceType::Vulkan:
+			FAVDevice::GetHardwareDevice()->SetContext<FVideoContextVulkan>(
+				MakeShared<FVideoContextVulkan>(
+					GetIVulkanDynamicRHI()->RHIGetVkInstance(),
+					GetIVulkanDynamicRHI()->RHIGetVkDevice(),
+					GetIVulkanDynamicRHI()->RHIGetVkPhysicalDevice(),
+					[](char const* Name) -> PFN_vkVoidFunction {
+						return (PFN_vkVoidFunction)GetIVulkanDynamicRHI()->RHIGetVkDeviceProcAddr(Name);
+					}));
+
+			break;
+#endif
+#if AVCODECS_USE_D3D
+		case ERHIInterfaceType::D3D11:
+			FAVDevice::GetHardwareDevice()->SetContext<FVideoContextD3D11>(
+				MakeShared<FVideoContextD3D11>(
+					GetID3D11DynamicRHI()->RHIGetDevice()));
+
+			break;
+		case ERHIInterfaceType::D3D12:
+			FAVDevice::GetHardwareDevice()->SetContext<FVideoContextD3D12>(
+				MakeShared<FVideoContextD3D12>(
+					GetID3D12DynamicRHI()->RHIGetDevice(0)));
+
+			break;
+#endif
+#if AVCODECS_USE_METAL
+		case ERHIInterfaceType::Metal:
+			FAVDevice::GetHardwareDevice()->SetContext<FVideoContextMetal>(
+				MakeShared<FVideoContextMetal>(
+					static_cast<MTL::Device*>(GDynamicRHI->RHIGetNativeDevice())));
+			break;
+#endif
+		default:
+			break;
+		}
+
+		FAVDevice::GetHardwareDevice()->SetContext<FVideoContextRHI>(MakeShared<FVideoContextRHI>());
+	}
+
+private:
+	static FDelegateHandle PostEngineInitHandle;
 };
 
+FDelegateHandle FAVCodecsCoreRHI::PostEngineInitHandle;
+
 IMPLEMENT_MODULE(FAVCodecsCoreRHI, AVCodecsCoreRHI);
diff --git a/Engine/Plugins/Experimental/RemoteSession/RemoteSession.uplugin b/Engine/Plugins/Experimental/RemoteSession/RemoteSession.uplugin
--- a/Engine/Plugins/Experimental/RemoteSession/RemoteSession.uplugin
+++ b/Engine/Plugins/Experimental/RemoteSession/RemoteSession.uplugin
@@ -56,7 +56,13 @@
 			"Name": "PixelStreaming2",
 			"Enabled": true,
 			"TargetAllowList": [
-				"Editor"
+				"Editor",
+				"Client"
+			],
+			"PlatformAllowList": [
+				"Win64",
+				"Mac",
+				"Linux"
 			]
 		}
 	]
diff --git a/Engine/Source/Programs/UnrealBuildTool/Configuration/Descriptors/PluginReferenceDescriptor.cs b/Engine/Source/Programs/UnrealBuildTool/Configuration/Descriptors/PluginReferenceDescriptor.cs
--- a/Engine/Source/Programs/UnrealBuildTool/Configuration/Descriptors/PluginReferenceDescriptor.cs
+++ b/Engine/Source/Programs/UnrealBuildTool/Configuration/Descriptors/PluginReferenceDescriptor.cs
@@ -74,6 +74,11 @@ namespace UnrealBuildTool
 		/// </summary>
 		public TargetType[]? TargetDenyList;
 
+		/// <summary>
+		/// Overrides the loading phase in uplugin so that a plugin can be explicitly loaded for specified targets.
+		/// </summary>
+		public TargetType[]? TargetsToExplicitlyLoad;
+
 		/// <summary>
 		/// The list of supported platforms for this plugin. This field is copied from the plugin descriptor, and supplements the user's allowed/denied platforms.
 		/// </summary>
@@ -153,6 +158,10 @@ namespace UnrealBuildTool
 			{
 				Writer.WriteEnumArrayField("TargetDenyList", TargetDenyList);
 			}
+			if (TargetsToExplicitlyLoad != null && TargetsToExplicitlyLoad.Length > 0)
+			{
+				Writer.WriteEnumArrayField("TargetsToExplicitlyLoad", TargetsToExplicitlyLoad);
+			}
 			if (SupportedTargetPlatforms != null && SupportedTargetPlatforms.Length > 0)
 			{
 				Writer.WriteStringArrayField("SupportedTargetPlatforms", SupportedTargetPlatforms.Select(x => x.ToString()).ToArray());
@@ -213,6 +222,10 @@ namespace UnrealBuildTool
 			{
 				PluginReferenceObject.AddOrSetFieldValue("TargetDenyList", TargetDenyList);
 			}
+			if (TargetsToExplicitlyLoad != null && TargetsToExplicitlyLoad.Length > 0)
+			{
+				PluginReferenceObject.AddOrSetFieldValue("TargetsToExplicitlyLoad", TargetsToExplicitlyLoad);
+			}
 			if (SupportedTargetPlatforms != null && SupportedTargetPlatforms.Length > 0)
 			{
 				PluginReferenceObject.AddOrSetFieldValue("SupportedTargetPlatforms", SupportedTargetPlatforms.Select(x => x.ToString()).ToArray());
@@ -309,6 +322,7 @@ namespace UnrealBuildTool
 					Descriptor.TargetConfigurationAllowList = RawObject.GetEnumArrayFieldWithDeprecatedFallback<UnrealTargetConfiguration>("TargetConfigurationAllowList", "WhitelistTargetConfigurations");
 					Descriptor.TargetConfigurationDenyList = RawObject.GetEnumArrayFieldWithDeprecatedFallback<UnrealTargetConfiguration>("TargetConfigurationDenyList", "BlacklistTargetConfigurations");
 
+					RawObject.TryGetEnumArrayField<TargetType>("TargetsToExplicitlyLoad", out Descriptor.TargetsToExplicitlyLoad);
 					RawObject.TryGetStringArrayField("SupportedTargetPlatforms", out Descriptor.SupportedTargetPlatforms);
 					RawObject.TryGetBoolField("HasExplicitPlatforms", out Descriptor.bHasExplicitPlatforms);
 
diff --git a/Engine/Source/Runtime/CUDA/Source/Private/CudaModule.cpp b/Engine/Source/Runtime/CUDA/Source/Private/CudaModule.cpp
--- a/Engine/Source/Runtime/CUDA/Source/Private/CudaModule.cpp
+++ b/Engine/Source/Runtime/CUDA/Source/Private/CudaModule.cpp
@@ -23,6 +23,8 @@
 
 #include "HAL/UnrealMemory.h"
 
+#include "Engine/Engine.h"
+
 DEFINE_LOG_CATEGORY(LogCUDA);
 
 // CUDA 11.0.0 is our minimum required version
@@ -30,17 +32,34 @@ DEFINE_LOG_CATEGORY(LogCUDA);
 
 void FCUDAModule::StartupModule()
 {
-	// Attempt to load the CUDA library and wire up our post-init delegate if loading was successful
+	// Attempt to load the CUDA library, if library loading was successful:
+	// When the plugin loading happens in the engine initialization, wire up our post-init delegate.
+	// If the plugin loading is on demand, calls InitCuda directly.
 	if (LoadCuda() == true)
 	{
-		FCoreDelegates::OnPostEngineInit.AddRaw(this, &FCUDAModule::InitCuda);
-		UE_LOG(LogCUDA, Display, TEXT("CUDA module ready pending PostEngineInit."));
+		if (!GEngine || !GEngine->IsInitialized())
+		{
+			PostEngineInitHandle = FCoreDelegates::OnPostEngineInit.AddRaw(this, &FCUDAModule::InitCuda);
+
+			UE_LOG(LogCUDA, Display, TEXT("CUDA module ready pending PostEngineInit."));
+		}
+		else
+		{
+			FCUDAModule::InitCuda();
+			UE_LOG(LogCUDA, Display, TEXT("CUDA module loading directly as PostEngineInit has passed."));
+		}
 	}
 }
 
 void FCUDAModule::ShutdownModule()
 {
 	UnloadCuda();
+
+	if (PostEngineInitHandle.IsValid())
+	{
+		FCoreDelegates::OnPostEngineInit.Remove(PostEngineInitHandle);
+		PostEngineInitHandle.Reset();
+	}
 }
 
 bool FCUDAModule::IsAvailable()
@@ -234,6 +253,13 @@ bool FCUDAModule::IsRHISelectedDevice(CUdevice cuDevice)
 
 void FCUDAModule::InitCuda()
 {
+	// Remove the binding
+	if (PostEngineInitHandle.IsValid())
+	{
+		FCoreDelegates::OnPostEngineInit.Remove(PostEngineInitHandle);
+		PostEngineInitHandle.Reset();
+	}
+
 	// Unload CUDA if we are not running the RHI on an NVIDIA device this can happen if the NVIDIA drivers
 	// are present on a system with an AMD GPU
 	// TODO is there the potential case that someone will want to use CUDA on a 2nd GPU
diff --git a/Engine/Source/Runtime/CUDA/Source/Public/CudaModule.h b/Engine/Source/Runtime/CUDA/Source/Public/CudaModule.h
--- a/Engine/Source/Runtime/CUDA/Source/Public/CudaModule.h
+++ b/Engine/Source/Runtime/CUDA/Source/Public/CudaModule.h
@@ -56,4 +56,6 @@ private:
 
 	uint32 rhiDeviceIndex;
 	TMap<uint32, CUcontext> contextMap;
+
+	FDelegateHandle PostEngineInitHandle;
 };
diff --git a/Engine/Source/Runtime/Projects/Private/PluginManager.cpp b/Engine/Source/Runtime/Projects/Private/PluginManager.cpp
--- a/Engine/Source/Runtime/Projects/Private/PluginManager.cpp
+++ b/Engine/Source/Runtime/Projects/Private/PluginManager.cpp
@@ -2375,6 +2375,12 @@ bool FPluginManager::ConfigureEnabledPluginForTarget(const FPluginReferenceDescr
 		return true;
 	}
 
+	if (FirstReference.TargetsToExplicitlyLoad.Contains(TargetType))
+	{
+		UE_LOG(LogPluginManager, Verbose, TEXT("Ignoring plugin '%s' for being set to load-on-demand in project."), *FirstReference.Name);
+		return true;
+	}
+
 	// Set of plugin names we've added to the queue for processing, and the plugin that referenced
 	// them, so we can display the chain of references if necessary for diagnostics
 	TMap<FString, const FPluginReferenceDescriptor*> SeenPlugins;
diff --git a/Engine/Source/Runtime/Projects/Private/PluginReferenceDescriptor.cpp b/Engine/Source/Runtime/Projects/Private/PluginReferenceDescriptor.cpp
--- a/Engine/Source/Runtime/Projects/Private/PluginReferenceDescriptor.cpp
+++ b/Engine/Source/Runtime/Projects/Private/PluginReferenceDescriptor.cpp
@@ -161,6 +161,9 @@ TOptional<FText> UE::Projects::Private::Read(UE::Json::FConstObject Object, FPlu
 	TryGetEnumArrayFieldWithDeprecatedFallback(Object, TEXT("TargetAllowList"), TEXT("WhitelistTargets"), /*out*/ Result.TargetAllowList);
 	TryGetEnumArrayFieldWithDeprecatedFallback(Object, TEXT("TargetDenyList"), TEXT("BlacklistTargets"), /*out*/ Result.TargetDenyList);
 
+	// Get the target list for explicitly loading targets
+	TryGetEnumArrayField(Object, TEXT("TargetsToExplicitlyLoad"), /*out*/ Result.TargetsToExplicitlyLoad);
+
 	// Get the supported platform list
 	TryGetStringArrayField(Object, TEXT("SupportedTargetPlatforms"), Result.SupportedTargetPlatforms);
 	TryGetBoolField(Object, TEXT("HasExplicitPlatforms"), Result.bHasExplicitPlatforms);
@@ -344,6 +347,20 @@ void FPluginReferenceDescriptor::UpdateJson(FJsonObject& JsonObject) const
 		JsonObject.RemoveField(TEXT("TargetDenyList"));
 	}
 
+	if (TargetsToExplicitlyLoad.Num() > 0)
+	{
+		TArray<TSharedPtr<FJsonValue>> TargetsToExplicitlyLoadValues;
+		for (EBuildTargetType Target : TargetsToExplicitlyLoad)
+		{
+			TargetsToExplicitlyLoadValues.Add(MakeShareable(new FJsonValueString(LexToString(Target))));
+		}
+		JsonObject.SetArrayField(TEXT("TargetsToExplicitlyLoad"), TargetsToExplicitlyLoadValues);
+	}
+	else
+	{
+		JsonObject.RemoveField(TEXT("TargetsToExplicitlyLoad"));
+	}
+
 	if (SupportedTargetPlatforms.Num() > 0)
 	{
 		TArray<TSharedPtr<FJsonValue>> SupportedTargetPlatformValues;
diff --git a/Engine/Source/Runtime/Projects/Public/PluginReferenceDescriptor.h b/Engine/Source/Runtime/Projects/Public/PluginReferenceDescriptor.h
--- a/Engine/Source/Runtime/Projects/Public/PluginReferenceDescriptor.h
+++ b/Engine/Source/Runtime/Projects/Public/PluginReferenceDescriptor.h
@@ -61,6 +61,9 @@ struct FPluginReferenceDescriptor
 	/** If enabled, list of targets for which the plugin should be disabled. */
 	TArray<EBuildTargetType> TargetDenyList;
 
+	/** Overrides the loading phase in uplugin so that a plugin can be explicitly loaded for specified targets. */
+	TArray<EBuildTargetType> TargetsToExplicitlyLoad;
+
 	/** The list of supported target platforms for this plugin. This field is copied from the plugin descriptor, and supplements the user's allowed/denied platforms. */
 	TArray<FString> SupportedTargetPlatforms;
 
