PixelStreaming2: Use MTL::ResourceStorageModeShared staging buffer for FVideoResourceMetal::CopyFrom to fix alignment/stride issues on iOS causing video stretching
#Synced-CL 51488969

diff --git a/Engine/Plugins/Experimental/AVCodecs/AVCodecsCore/Source/AVCodecsCore/Private/Video/Resources/Metal/VideoResourceMetal.cpp b/Engine/Plugins/Experimental/AVCodecs/AVCodecsCore/Source/AVCodecsCore/Private/Video/Resources/Metal/VideoResourceMetal.cpp
--- a/Engine/Plugins/Experimental/AVCodecs/AVCodecsCore/Source/AVCodecsCore/Private/Video/Resources/Metal/VideoResourceMetal.cpp
+++ b/Engine/Plugins/Experimental/AVCodecs/AVCodecsCore/Source/AVCodecsCore/Private/Video/Resources/Metal/VideoResourceMetal.cpp
@@ -88,11 +88,17 @@ FVideoDescriptor FVideoResourceMetal::GetDescriptorFrom(TSharedRef<FAVDevice> co
 FVideoResourceMetal::FVideoResourceMetal(TSharedRef<FAVDevice> const& Device, MTL::Texture* Raw, FAVLayout const& Layout)
 	: TVideoResource(Device, Layout, GetDescriptorFrom(Device, Raw))
 	, Raw(Raw)
+	, StagingBuffer(nullptr)
 {
 }
 
 FVideoResourceMetal::~FVideoResourceMetal()
 {
+	if (StagingBuffer != nullptr)
+	{
+		StagingBuffer->release();
+		StagingBuffer = nullptr;
+	}
 }
 
 FAVResult FVideoResourceMetal::CopyFrom(CVPixelBufferRef Other)
@@ -116,7 +122,27 @@ FAVResult FVideoResourceMetal::CopyFrom(CVPixelBufferRef Other)
 	{
 		return BytesPerPixel;
 	}
-	
+		
+	const size_t TightRowSize = static_cast<size_t>(BytesPerPixel) * SrcWidth;
+	const size_t BufferSize = TightRowSize * SrcHeight;
+
+	// Create staging buffer for the copy
+	if (StagingBuffer != nullptr && StagingBuffer->length() != BufferSize)
+	{
+		// We expect the FVideoResourceMetal to be recreated if the size ever changes, but we check here just to make sure
+		StagingBuffer->release();
+		StagingBuffer = nullptr;
+	}
+
+	if (StagingBuffer == nullptr)
+	{
+		StagingBuffer = GetContext()->Device->newBuffer(BufferSize, MTL::ResourceStorageModeShared);
+		if (StagingBuffer == nullptr)
+		{
+			return FAVResult(EAVResult::Error, TEXT("Failed to create staging buffer"), TEXT("Metal"));
+		}
+	}
+
 	if (const CVReturn Result = CVPixelBufferLockBaseAddress(Other, 0); Result != kCVReturnSuccess)
 	{
 		return FAVResult(EAVResult::Error, TEXT("Failed to lock input pixel buffer!"), TEXT("Metal"));
@@ -124,22 +150,17 @@ FAVResult FVideoResourceMetal::CopyFrom(CVPixelBufferRef Other)
 
 	const size_t SrcStride = CVPixelBufferGetBytesPerRow(Other);
 	const uint8* SrcPtr = static_cast<const uint8*>(CVPixelBufferGetBaseAddress(Other));
-	
-	const size_t TightRowSize = BytesPerPixel * SrcWidth;
-	
-	TArray<uint8> PackedData;
-	PackedData.SetNumZeroed(TightRowSize * SrcHeight);
-	
-	uint8* DstPtr = PackedData.GetData();
+
+	uint8* DstPtr = static_cast<uint8*>(StagingBuffer->contents());
 	for (size_t Row = 0; Row < SrcHeight; Row++)
 	{
 		const uint8* SrcRow = SrcPtr + Row * SrcStride;
 		uint8* DstRow = DstPtr + Row * TightRowSize;
-		
+
 		FMemory::Memcpy(DstRow, SrcRow, TightRowSize);
 	}
-	
-	Raw->replaceRegion(MTL::Region(0, 0, SrcWidth, SrcHeight), 0, 0, PackedData.GetData(), TightRowSize, 0);
+
+	Raw->replaceRegion(MTL::Region(0, 0, SrcWidth, SrcHeight), 0, 0, StagingBuffer->contents(), TightRowSize, 0);
 	
 	CVPixelBufferUnlockBaseAddress(Other, 0);
 	
diff --git a/Engine/Plugins/Experimental/AVCodecs/AVCodecsCore/Source/AVCodecsCore/Public/Video/Resources/Metal/VideoResourceMetal.h b/Engine/Plugins/Experimental/AVCodecs/AVCodecsCore/Source/AVCodecsCore/Public/Video/Resources/Metal/VideoResourceMetal.h
--- a/Engine/Plugins/Experimental/AVCodecs/AVCodecsCore/Source/AVCodecsCore/Public/Video/Resources/Metal/VideoResourceMetal.h
+++ b/Engine/Plugins/Experimental/AVCodecs/AVCodecsCore/Source/AVCodecsCore/Public/Video/Resources/Metal/VideoResourceMetal.h
@@ -37,6 +37,7 @@ class FVideoResourceMetal : public TVideoResource<FVideoContextMetal>
 {
 private:
     MTL::Texture* Raw;
+	MTL::Buffer* StagingBuffer;
 
 public:
 	static AVCODECSCORE_API FVideoDescriptor GetDescriptorFrom(TSharedRef<FAVDevice> const& Device, MTL::Texture* Raw);
@@ -49,6 +50,12 @@ public:
 	AVCODECSCORE_API FAVResult CopyFrom(CVPixelBufferRef Other);
     
 	AVCODECSCORE_API virtual FAVResult Validate() const override;
+
+	// Non-copyable and non-movable to prevent double-release of Metal resources.
+	FVideoResourceMetal(FVideoResourceMetal const&) = delete;
+	FVideoResourceMetal(FVideoResourceMetal&&) = delete;
+	FVideoResourceMetal& operator=(FVideoResourceMetal const&) = delete;
+	FVideoResourceMetal& operator=(FVideoResourceMetal&&) = delete;
 };
 
 DECLARE_TYPEID(FVideoContextMetal, AVCODECSCORE_API);
